Moved XmlParser from 'Developer' to 'Runtime' so it can be used in shipping games

diff --git a/Engine/Source/Developer/XmlParser/Private/XmlCharacterWidthCheck.h b/Engine/Source/Developer/XmlParser/Private/XmlCharacterWidthCheck.h
deleted file mode 100644
diff --git a/Engine/Source/Developer/XmlParser/Private/XmlFile.cpp b/Engine/Source/Developer/XmlParser/Private/XmlFile.cpp
deleted file mode 100644
diff --git a/Engine/Source/Developer/XmlParser/Private/XmlNode.cpp b/Engine/Source/Developer/XmlParser/Private/XmlNode.cpp
deleted file mode 100644
diff --git a/Engine/Source/Developer/XmlParser/Private/XmlParser.cpp b/Engine/Source/Developer/XmlParser/Private/XmlParser.cpp
deleted file mode 100644
diff --git a/Engine/Source/Developer/XmlParser/Private/XmlParserPrivatePCH.h b/Engine/Source/Developer/XmlParser/Private/XmlParserPrivatePCH.h
deleted file mode 100644
diff --git a/Engine/Source/Developer/XmlParser/Public/XmlFile.h b/Engine/Source/Developer/XmlParser/Public/XmlFile.h
deleted file mode 100644
diff --git a/Engine/Source/Developer/XmlParser/Public/XmlNode.h b/Engine/Source/Developer/XmlParser/Public/XmlNode.h
deleted file mode 100644
diff --git a/Engine/Source/Developer/XmlParser/Public/XmlParser.h b/Engine/Source/Developer/XmlParser/Public/XmlParser.h
deleted file mode 100644
diff --git a/Engine/Source/Developer/XmlParser/XmlParser.Build.cs b/Engine/Source/Developer/XmlParser/XmlParser.Build.cs
deleted file mode 100644
diff --git a/Engine/Source/Runtime/XmlParser/Private/XmlCharacterWidthCheck.h b/Engine/Source/Runtime/XmlParser/Private/XmlCharacterWidthCheck.h
new file mode 100644
--- /dev/null
+++ b/Engine/Source/Runtime/XmlParser/Private/XmlCharacterWidthCheck.h
@@ -0,0 +1,85 @@
+// Copyright 1998-2014 Epic Games, Inc. All Rights Reserved.
+
+#pragma once
+
+/**
+ * Helper class to determine the character width used by an XML file
+ */
+class CharacterWidthCheck
+{
+public:
+	/**
+	 * Set buffer containing XML aata
+	 * @note Buffer must be a least 4 bytes big
+	 * @parameter Buffer to check for XML
+	 **/
+	CharacterWidthCheck(const void* Buffer)
+		: TextStart(Buffer)
+		, CharacterWidth(1)
+		, EntireBuffer(Buffer)
+	{
+	}
+
+	/**
+	 * Check start of buffer and store determined character width
+	 * @return Whether the buffer appears to contain XML
+	 */
+	bool FindCharacterWidth()
+	{
+		// http://en.wikipedia.org/wiki/Byte_order_mark
+		auto BufferBytes = static_cast<const uint8*>(EntireBuffer);
+		if (BufferBytes[0] == 0xef && BufferBytes[1] == 0xbb && BufferBytes[2] == 0xbf)
+		{
+			// Accept UTF-8: we don't attempt to decode it, just pass it through unscathed
+			TextStart = &BufferBytes[3];
+			return true;
+		}
+
+		if (CheckBOMOrOpenTag<uint16>() || CheckBOMOrOpenTag<uint32>())
+		{
+			return true;
+		}
+
+		if ((BufferBytes[0] == '<') && (BufferBytes[1] != 0))
+		{
+			// Plain 8 bit characters
+			TextStart = BufferBytes;
+			CharacterWidth = 1;
+			return true;
+		}
+
+		return false;
+	}
+
+	/** Start of actual XML, skipping BOM if present */
+	const void* TextStart;
+
+	/** Number of bytes each character in the XML data takes up */
+	int CharacterWidth;
+
+private:
+	/**
+	 * Check for 16 or 32-bit XML data
+	 * @return Whether the buffer appears to contain XML
+	 */
+	template <typename CharType>
+	bool CheckBOMOrOpenTag()
+	{
+		auto FirstWideChar = static_cast<const CharType*>(EntireBuffer);
+		if (*FirstWideChar != UNICODE_BOM && *FirstWideChar != '<')
+		{
+			return false;
+		}
+
+		if (*FirstWideChar == UNICODE_BOM)
+		{
+			TextStart = FirstWideChar + 1;
+		}
+		CharacterWidth = sizeof(CharType);
+		return true;
+	}
+
+	/** Provided buffer that may or may not contain XML */
+	const void* EntireBuffer;
+};
+
diff --git a/Engine/Source/Runtime/XmlParser/Private/XmlFile.cpp b/Engine/Source/Runtime/XmlParser/Private/XmlFile.cpp
new file mode 100644
--- /dev/null
+++ b/Engine/Source/Runtime/XmlParser/Private/XmlFile.cpp
@@ -0,0 +1,889 @@
+// Copyright 1998-2014 Epic Games, Inc. All Rights Reserved.
+
+#include "XmlParserPrivatePCH.h"
+
+#include "XmlFile.h"
+#include "UniquePtr.h"
+#include "XmlCharacterWidthCheck.h"
+
+namespace
+{
+/**
+ * Split a buffer of characters into lines
+ * @param OutArray Array of strings to write lines to
+ * @param InBuffer Memory containing text to split
+ * @param EndOfBuffer Pointer to the byte after last byte of the buffer
+ * @param Delim Character used to determine a line ending
+ */
+template <typename CharType>
+void SplitLines(TArray<FString>& OutArray, const CharType* InBuffer, const void* EndOfBuffer, ANSICHAR Delim = '\n');
+
+/**
+ * Take an XML buffer, detect the size of character it uses and split into lines
+ * @param OutArray Array of strings to write lines to
+ * @param InBuffer Memory containing text to split
+ * @return False if character type not detected (only guaranteed for XML files)
+ */
+bool FindCharSizeAndSplitLines(TArray<FString>& OutArray, const void* InBuffer, uint32 BufferSize);
+
+}
+
+
+FXmlFile::FXmlFile(const FString& InFile, EConstructMethod::Type ConstructMethod)
+	: RootNode(nullptr), bFileLoaded(false)
+{
+	LoadFile(InFile, ConstructMethod);
+}
+
+bool FXmlFile::LoadFile(const FString& InFile, EConstructMethod::Type ConstructMethod)
+{
+	// Remove any file stuff if it already exists
+	Clear();
+
+	// So far no error (Early so it can be overwritten below by errors)
+	ErrorMessage = NSLOCTEXT("XmlParser", "LoadSuccess", "XmlFile was loaded successfully").ToString();
+
+	TArray<FString> Input;
+	if(ConstructMethod == EConstructMethod::ConstructFromFile)
+	{
+		// Create file reader
+		TUniquePtr<FArchive> FileReader(IFileManager::Get().CreateFileReader(*InFile));
+		if(!FileReader)
+		{
+			ErrorMessage = NSLOCTEXT("XmlParser", "FileLoadFail", "Failed to load the file").ToString();
+			ErrorMessage += TEXT("\"");
+			ErrorMessage += InFile;
+			ErrorMessage += TEXT("\"");
+			return false;
+		}
+
+		// Create buffer for file input
+		uint32 BufferSize = FileReader->TotalSize();
+		void* Buffer = FMemory::Malloc(BufferSize);
+		FileReader->Serialize(Buffer, BufferSize);
+
+		// Parse file buffer into an array of lines
+		if (!FindCharSizeAndSplitLines(Input, Buffer, BufferSize))
+		{
+			ErrorMessage = NSLOCTEXT("XmlParser", "InvalidFormatFail", "Failed to parse the file (Unsupported character encoding)").ToString();
+			ErrorMessage += TEXT("\"");
+			ErrorMessage += InFile;
+			ErrorMessage += TEXT("\"");
+			return false;
+		}
+
+		// Release resources
+		FMemory::Free(Buffer);
+	}
+	else
+	{
+		// Parse input buffer into an array of lines
+		SplitLines(Input, *InFile, *InFile + InFile.Len());
+	}
+
+	// Pre-process the input
+	PreProcessInput(Input);
+
+	// Tokenize the input
+	TArray<FString> Tokens = Tokenize(Input);
+
+	// Parse the input & create the nodes
+	CreateNodes(Tokens);
+
+	// All done with creation, set up necessary information
+	if(bFileLoaded == true)
+	{
+		if(ConstructMethod == EConstructMethod::ConstructFromFile)
+		{
+			LoadedFile = InFile;
+		}
+	}
+	else
+	{
+		LoadedFile = TEXT("");
+		RootNode = nullptr;
+	}
+
+	// Now check the status flag of the creation. It may have actually failed, but given us a 
+	// partially created representation
+	if(bCreationFailed)
+	{
+		Clear();
+	}
+
+	return bFileLoaded;
+}
+
+FString FXmlFile::GetLastError() const
+{
+	return ErrorMessage;
+}
+
+void FXmlFile::Clear()
+{
+	if(bFileLoaded)
+	{
+		check(RootNode != nullptr);
+		RootNode->Delete();
+		delete RootNode;
+
+		RootNode = nullptr;
+		bFileLoaded = false;
+		LoadedFile = TEXT("");
+		ErrorMessage = NSLOCTEXT("XmlParser", "ClearSuccess", "XmlFile was cleared successfully").ToString();
+	}
+}
+
+bool FXmlFile::IsValid() const
+{
+	checkSlow(bFileLoaded ? RootNode != nullptr : RootNode == nullptr);
+	return bFileLoaded;
+}
+
+const FXmlNode* FXmlFile::GetRootNode() const
+{
+	return RootNode;
+}
+
+FXmlNode* FXmlFile::GetRootNode()
+{
+	return RootNode;
+}
+
+bool FXmlFile::Save(const FString& Path)
+{
+	FString Xml;
+	Xml += TEXT("<?xml version=\"1.0\" encoding=\"UTF-8\"?>") LINE_TERMINATOR;
+
+	TArray<const FXmlNode*> NodeStack;
+	static const TCHAR Tabs[] = TEXT("\t\t\t\t\t\t\t\t\t\t\t\t\t\t\t\t");
+	auto Indent = &Tabs[ARRAY_COUNT(Tabs) - 1];
+	const auto* CurrentNode = GetRootNode();
+	while (CurrentNode)
+	{
+		if (auto Child = CurrentNode->GetFirstChildNode())
+		{
+			NodeStack.Push(CurrentNode);
+			Xml += FString::Printf(TEXT("%s<%s>") LINE_TERMINATOR, Indent, *CurrentNode->GetTag());
+
+			CurrentNode = Child;	
+			// Actually increases the indent by one tab
+			--Indent;
+			continue;
+		}
+
+		auto Tag = *CurrentNode->GetTag();
+		Xml += FString::Printf(TEXT("%s<%s>%s</%s>") LINE_TERMINATOR, Indent, Tag, *CurrentNode->GetContent(), Tag);
+		CurrentNode = CurrentNode->GetNextNode();
+
+		while (!CurrentNode && NodeStack.Num() != 0)
+		{
+			// Actually decreases the indent by one tab
+			++Indent;
+			CurrentNode = NodeStack.Pop();
+			Xml += FString::Printf(TEXT("%s</%s>") LINE_TERMINATOR, Indent, *CurrentNode->GetTag());
+			CurrentNode = CurrentNode->GetNextNode();
+		}
+	}
+	
+	TUniquePtr<FArchive> Archive(IFileManager::Get().CreateFileWriter(*Path));
+	if (!Archive)
+	{
+		ErrorMessage = NSLOCTEXT("XmlParser", "FileSaveFail", "Failed to save the file").ToString();
+		ErrorMessage += FString::Printf(TEXT("\"%s\""), *Path);
+		return false;
+	}
+
+	FTCHARToUTF8 Converter(*Xml);
+	Archive->Serialize(const_cast<char*>(Converter.Get()), Converter.Length());
+	return true;
+}
+
+namespace
+{
+template <typename CharType>
+void SplitLines(TArray<FString>& OutArray, const CharType* Buffer, const void* EndOfBuffer, ANSICHAR Delim)
+{
+	uint32 CharacterCount = static_cast<const CharType*>(EndOfBuffer) - Buffer;
+	FString WorkingLine;
+	for(uint32 i = 0; i != CharacterCount; ++i)
+	{
+		if(Buffer[i] == Delim)
+		{
+			if(WorkingLine.Len())
+			{
+				OutArray.Add(WorkingLine);
+				WorkingLine = TEXT("");
+			}
+		}
+		else
+		{
+			WorkingLine += static_cast<TCHAR>(Buffer[i]);
+		}
+	}
+
+	if(WorkingLine.Len())
+	{
+		OutArray.Add(WorkingLine);
+	}
+}
+
+bool FindCharSizeAndSplitLines(TArray<FString>& OutArray, const void* InBuffer, uint32 BufferSize)
+{
+	if (BufferSize < 4)
+	{
+		// Ensure there's enough buffer for CharacterWidthCheck to use (4 characters is not enough
+		// to store any XML)
+		return false;
+	}
+
+	CharacterWidthCheck CharCheck(InBuffer);
+	if (!CharCheck.FindCharacterWidth())
+	{
+		return false;
+	}
+
+	const void* EndOfBuffer = static_cast<const char*>(InBuffer) + BufferSize;
+
+	// Parse input buffer into an array of lines
+	switch (CharCheck.CharacterWidth)
+	{
+	case 1:
+		SplitLines(OutArray, static_cast<const uint8*>(CharCheck.TextStart), EndOfBuffer);
+		break;
+
+	case 2:
+		SplitLines(OutArray, static_cast<const uint16*>(CharCheck.TextStart), EndOfBuffer);
+		break;
+
+	case 4:
+		SplitLines(OutArray, static_cast<const uint32*>(CharCheck.TextStart), EndOfBuffer);
+		break;
+
+	default:
+		return false;
+	}
+	return true;
+}
+
+} // end anonymous namespace
+
+/** Checks if the passed character is a whitespace character */
+static bool IsWhiteSpace(TCHAR Char)
+{
+	// Whitespace will be any character that is not a common printed ASCII character (and not space/tab)
+	if(Char == TCHAR(' ') ||
+		Char == TCHAR('\t') ||
+		Char < 32 ) // ' '
+	{
+		return true;
+	}
+	return false;
+}
+
+void FXmlFile::PreProcessInput(TArray<FString>& Input)
+{
+	// Note: This implementation is written simply and will not handle all cases.  It
+	//       is made for the simple cases where FXmlFile is to be used.
+	// Assumptions/Misc:
+	//       - Well-formatted file with 1 entry per line
+	//       - Ignoring versions, encodings, and doctypes
+
+	// Remove white space at the beginning of lines
+	for(int32 i = 0; i < Input.Num(); ++i)
+	{
+		int32 NumWhiteSpace = 0;
+		for(int32 j = 0; j < Input[i].Len(); ++j)
+		{
+			if(!IsWhiteSpace(Input[i][j]))
+			{
+				break;
+			}
+			
+			++NumWhiteSpace;
+		}
+
+		if(NumWhiteSpace > 0)
+		{
+			Input[i] = Input[i].Mid(NumWhiteSpace);
+		}
+	}
+
+	// Cull any text that can be removed on a line-based parse
+	for(int32 i = 0; i < Input.Num(); ++i)
+	{
+		// Find <!DOCTYPE or <?xml and remove those lines
+		if(Input[i].StartsWith(TEXT("<!DOCTYPE")) || Input[i].StartsWith(TEXT("<?xml")))
+		{
+			Input[i] = TEXT("");
+		}
+	}
+
+	// Cull any text inside of comments
+	bool bInComment = false;
+	int32 CommentLineStart = -1;
+	int32 CommentIndexStart = -1;
+	for(int32 i = 0; i < Input.Num(); ++i)
+	{
+		if(Input[i].Len() == 3)
+		{
+			if(bInComment)
+			{
+				if(Input[i][0] == TCHAR('-') && Input[i][1] == TCHAR('-') && Input[i][2] == TCHAR('>'))
+				{
+					// Found comment end, perform removal (simply replace all text with whitespace to be ignored by tokenizer)
+					bInComment = false;
+					int32 CommentLineEnd = i;
+					int32 CommentIndexEnd = 2;
+					WhiteOut(Input, CommentLineStart, CommentLineEnd, CommentIndexStart, CommentIndexEnd);
+				}
+			}
+		}
+
+		if(Input[i].Len() < 3)
+		{
+			continue;
+		}
+
+		int32 Indx1 = 0, Indx2 = 1, Indx3 = 2, Indx4 = 3;
+		for(; Indx4 < Input[i].Len(); ++Indx1, ++Indx2, ++Indx3, ++Indx4)
+		{
+			// Looking for the start of a comment
+			if(!bInComment)
+			{
+				if(Input[i][Indx1] == TCHAR('<') && Input[i][Indx2] == TCHAR('!') && Input[i][Indx3] == TCHAR('-') && Input[i][Indx4] == TCHAR('-'))
+				{
+					// Found comment start, mark it
+					bInComment = true;
+					CommentLineStart = i;
+					CommentIndexStart = Indx1;
+				}
+			}
+
+			// Looking for the end of a comment
+			else
+			{
+				if( (Input[i][Indx2] == TCHAR('-') && Input[i][Indx3] == TCHAR('-') && Input[i][Indx4] == TCHAR('>')) ||
+					(Input[i][Indx1] == TCHAR('-') && Input[i][Indx2] == TCHAR('-') && Input[i][Indx3] == TCHAR('>')) )
+				{
+					// Found comment end, perform removal (simply replace all text with whitespace to be ignored by tokenizer)
+					bInComment = false;
+					int32 CommentLineEnd = i;
+					int32 CommentIndexEnd = Indx4;
+					WhiteOut(Input, CommentLineStart, CommentLineEnd, CommentIndexStart, CommentIndexEnd);
+				}
+			}
+		}
+	}
+}
+
+void FXmlFile::WhiteOut(TArray<FString>& Input, int32 LineStart, int32 LineEnd, int32 IndexStart, int32 IndexEnd)
+{
+	if(LineEnd < LineStart)
+	{
+		// Error!, malformed file with comment ends
+	}
+
+	// White-out first line
+	if(LineEnd - LineStart > 0)
+	{
+		for(int32 i = IndexStart; i < Input[LineStart].Len(); ++i)
+		{
+			Input[LineStart][i] = TCHAR(' ');
+		}
+	}
+
+	// White-out middle lines
+	if(LineEnd - LineStart > 1)
+	{
+		int32 NumLines = LineEnd - LineStart - 1;
+		int32 MidStart = LineStart + 1;
+		for(int32 i = MidStart; i < MidStart + NumLines; ++i)
+		{
+			Input[i] = TEXT("");
+		}
+	}
+
+	// White-out last line
+	if(LineEnd - LineStart > 0)
+	{
+		for(int32 i = 0; i <= IndexEnd; ++i)
+		{
+			Input[LineEnd][i] = TCHAR(' ');
+		}
+	}
+
+	// White-out if comment is only on 1 line
+	if(LineStart == LineEnd)
+	{
+		for(int32 i = IndexStart; i <= IndexEnd; ++i)
+		{
+			Input[LineStart][i] = TCHAR(' ');
+		}
+	}
+}
+
+/** Checks if the passed character is an operator */
+static bool CheckTagOperator(const FString& InString, int32 InIndex)
+{
+	check(InIndex >= 0 && InIndex < InString.Len())
+	if(InString[InIndex] == TCHAR('/'))
+	{
+		if(InIndex < InString.Len() - 1 && InString[InIndex + 1] == TCHAR('>'))
+		{
+			return true;
+		}
+		// check either the next or previous chars are tag closures - otherwise, this is just a slash
+		else if(InIndex > 0 && InString[InIndex - 1] == TCHAR('<'))
+		{
+			return true;
+		}
+	}
+	else if(InString[InIndex] == TCHAR('<') || InString[InIndex] == TCHAR('>'))
+	{
+		return true;
+	}
+
+	return false;
+}
+
+/** Checks if the passed character is an operator */
+static bool IsPartOfTagOperator(TCHAR Char)
+{
+	if(Char == TCHAR('<') ||
+		Char == TCHAR('/') ||
+		Char == TCHAR('>'))
+	{
+		return true;
+	}
+	return false;
+}
+
+/** Checks if the passed character is a quote */
+static bool IsQuote(TCHAR Char)
+{
+	return Char == TCHAR('\"');
+}
+
+TArray<FString> FXmlFile::Tokenize(FString Input)
+{
+	TArray<FString> Tokens;
+
+	FString WorkingToken;
+	enum TOKENTYPE { OPERATOR, STRING, NONE } Type = NONE;
+	bool bInToken = false;
+	bool bInQuote = false;
+	for(int32 i = 0; i < Input.Len(); ++i)
+	{
+		if(IsWhiteSpace(Input[i]) && !bInQuote)
+		{
+			// End the current token 
+			if(WorkingToken.Len())
+			{
+				Tokens.Add(WorkingToken);
+				WorkingToken = TEXT("");
+			}
+			bInToken = false;
+			Type = NONE;
+
+			continue;
+		}
+
+		// Mark the start of a token
+		if(bInToken == false)
+		{
+			WorkingToken = TEXT("");
+			WorkingToken += Input[i];
+			bInToken = true;
+			if(CheckTagOperator(Input, i))
+			{
+				Type = OPERATOR;
+			}
+			else
+			{
+				Type = STRING;
+			}
+		}
+
+		// Already in a token, so continue parsing
+		else
+		{
+			if(Type == OPERATOR)
+			{
+				// Still the tag, so add it to the working token
+				if(CheckTagOperator(Input, i))
+				{
+					WorkingToken += Input[i];
+				}
+
+				// Not a tag operator anymore, so add the old token and start a new one
+				else
+				{
+					if(WorkingToken.Len())
+					{
+						Tokens.Add(WorkingToken);
+						WorkingToken = TEXT("");
+					}
+					WorkingToken += Input[i];
+					Type = STRING;
+				}
+			}
+			else // STRING
+			{
+				if(IsQuote(Input[i]) && !bInQuote)
+				{
+					bInQuote = true;
+				}
+				else if(IsQuote(Input[i]) && bInQuote)
+				{
+					bInQuote = false;
+				}
+
+				// Still a string
+				if(!CheckTagOperator(Input, i))
+				{
+					WorkingToken += Input[i];
+				}
+
+				// Moving back to operator
+				else
+				{
+					if(WorkingToken.Len())
+					{
+						Tokens.Add(WorkingToken);
+						WorkingToken = TEXT("");
+					}
+					WorkingToken += Input[i];
+					Type = OPERATOR;
+				}
+			}
+		}
+
+		// If we have a working token, add it if it's final (ie: ends with '>')
+		if(WorkingToken.Len() > 0)
+		{
+			if(WorkingToken[WorkingToken.Len() - 1] == TCHAR('>'))
+			{
+				Tokens.Add(WorkingToken);
+				WorkingToken = TEXT("");
+				bInToken = false;
+				Type = NONE;
+			}
+		}
+	}
+
+	// Add working token if it still exists
+	if(WorkingToken.Len())
+	{
+		Tokens.Add(WorkingToken);
+	}
+
+	// Return result
+	return Tokens;
+}
+
+TArray<FString> FXmlFile::Tokenize(TArray<FString>& Input)
+{
+	TArray<FString> Tokens;
+	for(int32 i = 0; i < Input.Num(); ++i)
+	{
+		Tokens.Append(Tokenize(Input[i]));
+	}
+	return Tokens;
+}
+
+/** Checks if the passed string is an important tag operator */
+static bool IsTagOperator(FString ToCheck)
+{
+	if(ToCheck == TEXT("<") ||
+		ToCheck == TEXT(">") ||
+		ToCheck == TEXT("</") ||
+		ToCheck == TEXT("/>"))
+	{
+		return true;
+	}
+	return false;
+}
+
+void FXmlFile::AddAttribute(const FString& InToken, TArray<FXmlAttribute>& OutAttributes)
+{
+	TArray<FString> Tokens;
+	if(InToken.Contains(TEXT("=")))
+	{
+		InToken.ParseIntoArray(&Tokens, TEXT("="), true);
+		if(Tokens.Num() == 2)
+		{
+			bool bQuotesRemoved = false;
+			FString Value = Tokens[1].TrimQuotes(&bQuotesRemoved);
+			if(bQuotesRemoved)
+			{
+				OutAttributes.Add(FXmlAttribute(Tokens[0], Value));
+			}
+		}
+	}
+}
+
+FXmlNode* FXmlFile::CreateNodeRecursive(const TArray<FString>& Tokens, int32 StartIndex, int32* OutNextIndex)
+{
+	// Algorithm (Draft):
+	//  - First found token should always be a '<'
+	//  - Extract tag
+	//  - Check next token
+	//    - If '<', recursively deepen
+	//    - If string token, parse until '<' found to recursively deepen
+	//  - Continue parsing until </tag> for self is found
+	//  - Return own constructed node (and index of next starting point
+
+	int32 SavedIndex = StartIndex;
+
+	// Get the tag & any attributes
+	FString Tag;
+	TArray<FXmlAttribute> Attributes;
+	bool bInTag = false;
+	for(int32 i = StartIndex; i < Tokens.Num(); ++i)
+	{
+		if(bCreationFailed)
+		{
+			break;
+		}
+
+		if(Tokens[i].Len() == 0 || Tokens[i] == TEXT("\n") || Tokens[i] == TEXT("\r") || Tokens[i] == TEXT("\n\r"))
+		{
+			continue;
+		}
+
+		// Looking for tag start
+		if(!bInTag)
+		{
+			if(Tokens[i] == TEXT("<"))
+			{
+				bInTag = true;
+			}
+
+			else
+			{
+				// Error: Found text before any operators (ex: plist>)
+				bCreationFailed = true;
+				ErrorMessage = NSLOCTEXT("XmlParser", "MalformedXMLFile", "Malformed Xml File").ToString();
+			}
+		}
+
+		// Already found tag start, reading inside of it now
+		else if(bInTag)
+		{
+			// Text for the tag
+			if(!IsTagOperator(Tokens[i]))
+			{
+				if(Tag.Len())
+				{
+					// this could be an attribute
+					AddAttribute(Tokens[i], Attributes);
+				}
+				else
+				{
+					Tag = Tokens[i];
+				}
+			}
+
+			// Found another tag operator
+			else
+			{
+				// Closing the tag starting
+				if(Tokens[i] == TEXT(">"))
+				{
+					SavedIndex = i + 1; // We want to find the content starting here possibly
+					break;
+				}
+
+				// Closing the tag and also finalizing it
+				else if(Tokens[i] == TEXT("/>"))
+				{
+					FXmlNode* NewNode = new FXmlNode();
+					NewNode->Tag = Tag;
+					NewNode->Attributes = Attributes;
+					check(OutNextIndex != nullptr);
+					*OutNextIndex = i + 1;
+					return NewNode;
+				}
+
+				// Error
+				else
+				{
+					// Error: malformed file (ex: <key<)
+					bCreationFailed = true;
+					ErrorMessage = NSLOCTEXT("XmlParser", "MalformedXMLFile", "Malformed Xml File").ToString();
+				}
+			}
+		}
+	}
+	
+	// Create a new node for the current tag
+	FXmlNode* NewNode = new FXmlNode();
+	NewNode->Tag = Tag;
+	NewNode->Attributes = Attributes;
+
+	// Got the tag. Continue and try to get the content
+	FString Content;
+	FString FinalTag;
+	bInTag = false;
+	for(int32 i = SavedIndex; i < Tokens.Num(); ++i)
+	{
+		if(bCreationFailed)
+		{
+			break;
+		}
+
+		if(Tokens[i].Len() == 0 || Tokens[i] == TEXT("\n") || Tokens[i] == TEXT("\r") || Tokens[i] == TEXT("\n\r"))
+		{
+			continue;
+		}
+
+		if(!bInTag)
+		{
+			// Found the start of another tag
+			if(Tokens[i] == TEXT("<"))
+			{
+				// Recursively enter function creating a child at the new tag
+				FXmlNode* Child = CreateNodeRecursive(Tokens, i, &SavedIndex);
+
+				// Save child to parent
+				if(Child != nullptr)
+				{
+					NewNode->Children.Add(Child);
+				}
+
+				// Error creating a child
+				else
+				{
+					// Some error occurred in creating a child. The algorithm should abort
+					bCreationFailed = true;
+					ErrorMessage = NSLOCTEXT("XmlParser", "MalformedXMLFile", "Malformed Xml File").ToString();
+					break;
+				}
+
+				// Continue algorithm at the returned index
+				i = SavedIndex - 1; // (-1) because continuing will increment i;
+				continue;
+			}
+
+			// Found what should be the end of the current tag
+			else if(Tokens[i] == TEXT("</"))
+			{
+				// Continue, getting the tag and checking to make sure it matched
+				bInTag = true;
+			}
+
+			// Error
+			else if(IsTagOperator(Tokens[i]))
+			{
+				// Error: Invalid token such as '<key>>'
+				bCreationFailed = true;
+				ErrorMessage = NSLOCTEXT("XmlParser", "MalformedXMLFile", "Malformed Xml File").ToString();
+			}
+
+			// Not an operator, save the text 
+			else
+			{
+				if(Content.Len() > 0)
+				{
+					Content += TEXT(" ");
+				}
+				Content += Tokens[i];
+			}
+		}
+
+		// In the finalizing tag
+		else
+		{
+			// Text for the tag
+			if(!IsTagOperator(Tokens[i]))
+			{
+				if(FinalTag.Len())
+				{
+					// Ignore this bit of text then (they may be like options or modifiers or some such)
+				}
+				else
+				{
+					FinalTag = Tokens[i];
+				}
+			}
+
+			else
+			{
+				// Found the end of the final tag
+				if(Tokens[i] == TEXT(">"))
+				{
+					// Error: Tag and ending tag don't match
+					if(Tag != FinalTag)
+					{
+						// Error: Tag and ending tag don't match
+						bCreationFailed = true;
+						ErrorMessage = NSLOCTEXT("XmlParser", "MalformedXMLFile", "Malformed Xml File").ToString();
+					}
+					
+					// Check to make sure tags match and create the node and return it
+					SavedIndex = i + 1;
+					if(OutNextIndex != nullptr) // Might be null because of the top-level node not sending it in
+					{
+						*OutNextIndex = SavedIndex;
+					}
+					NewNode->Content = Content;
+					return NewNode;
+				}
+
+				// Error
+				else
+				{
+					// Error: malformed file (ex: <key>stuff</key/>)
+					bCreationFailed = true;
+					ErrorMessage = NSLOCTEXT("XmlParser", "MalformedXMLFile", "Malformed Xml File").ToString();
+				}
+			}
+		}
+	}
+
+	// Something bad happened, return NULL!
+	delete NewNode;
+	return nullptr;
+}
+
+void FXmlFile::HookUpNextPtrs(FXmlNode* Node)
+{
+	if(Node == nullptr)
+	{
+		return;
+	}
+
+	for(int32 i = 0; i < Node->Children.Num(); ++i)
+	{
+		HookUpNextPtrs(Node->Children[i]);
+
+		if(i != Node->Children.Num() - 1)
+		{
+			Node->Children[i]->NextNode = Node->Children[i + 1];
+		}
+	}
+}
+
+void FXmlFile::CreateNodes(const TArray<FString>& Tokens)
+{
+	// Assumption..There is only 1 top-level node which contains everything inside of it
+	bCreationFailed = false;
+	FXmlNode* Root = CreateNodeRecursive(Tokens, 0);
+
+	if(Root)
+	{
+		bFileLoaded = true;
+
+		// Hook up next ptrs
+		HookUpNextPtrs(Root);
+
+		// Save it
+		RootNode = Root;
+	}
+	else
+	{
+		bFileLoaded = false;
+		ErrorMessage = NSLOCTEXT("XmlParser", "NodeCreateFail", "Failed to parse the loaded document").ToString();
+	}
+}
diff --git a/Engine/Source/Runtime/XmlParser/Private/XmlNode.cpp b/Engine/Source/Runtime/XmlParser/Private/XmlNode.cpp
new file mode 100644
--- /dev/null
+++ b/Engine/Source/Runtime/XmlParser/Private/XmlNode.cpp
@@ -0,0 +1,104 @@
+// Copyright 1998-2014 Epic Games, Inc. All Rights Reserved.
+
+#include "XmlParserPrivatePCH.h"
+
+#include "XmlNode.h"
+
+const FString& FXmlAttribute::GetTag() const
+{
+	return Tag;
+}
+
+const FString& FXmlAttribute::GetValue() const
+{
+	return Value;
+}
+
+void FXmlNode::Delete()
+{
+	const int32 ChildCount = Children.Num();
+	for(int32 ChildIndex = 0; ChildIndex < ChildCount; ++ChildIndex)
+	{
+		check(Children[ChildIndex] != nullptr);
+		Children[ChildIndex]->Delete();
+		delete Children[ChildIndex];
+	}
+	Children.Empty();
+}
+
+const FXmlNode* FXmlNode::GetNextNode() const
+{
+	return NextNode;
+}
+
+const TArray<FXmlNode*>& FXmlNode::GetChildrenNodes() const
+{
+	return Children;
+}
+
+const FXmlNode* FXmlNode::GetFirstChildNode() const
+{
+	if(Children.Num() > 0)
+	{
+		return Children[0];
+	}
+	else
+	{
+		return nullptr;
+	}
+}
+
+const FXmlNode* FXmlNode::FindChildNode(const FString& InTag) const
+{
+	const int32 ChildCount = Children.Num();
+	for(int32 ChildIndex = 0; ChildIndex < ChildCount; ++ChildIndex)
+	{
+		if(Children[ChildIndex] != nullptr && Children[ChildIndex]->GetTag() == InTag)
+		{
+			return Children[ChildIndex];
+		}
+	}
+
+	return nullptr;
+}
+
+FXmlNode* FXmlNode::FindChildNode(const FString& InTag)
+{
+	return const_cast<FXmlNode*>(const_cast<const FXmlNode*>(this)->FindChildNode(InTag));
+}
+
+const FString& FXmlNode::GetTag() const
+{
+	return Tag;
+}
+
+const FString& FXmlNode::GetContent() const
+{
+	return Content;
+}
+
+FString FXmlNode::GetAttribute(const FString& InTag) const
+{
+	for(auto Iter(Attributes.CreateConstIterator()); Iter; Iter++)
+	{
+		if(Iter->GetTag() == InTag)
+		{
+			return Iter->GetValue();
+		}
+	}
+	return FString();
+}
+
+void FXmlNode::AppendChildNode(const FString& InTag, const FString& InContent)
+{
+	auto NewNode = new FXmlNode;
+	NewNode->Tag = InTag;
+	NewNode->Content = InContent;
+
+	auto NumChildren = Children.Num();
+	if (NumChildren != 0)
+	{
+		Children[NumChildren - 1]->NextNode = NewNode;
+	}
+	Children.Push(NewNode);
+}
diff --git a/Engine/Source/Runtime/XmlParser/Private/XmlParser.cpp b/Engine/Source/Runtime/XmlParser/Private/XmlParser.cpp
new file mode 100644
--- /dev/null
+++ b/Engine/Source/Runtime/XmlParser/Private/XmlParser.cpp
@@ -0,0 +1,5 @@
+// Copyright 1998-2014 Epic Games, Inc. All Rights Reserved.
+
+#include "XmlParserPrivatePCH.h"
+
+IMPLEMENT_MODULE( FDefaultModuleImpl, XmlParser );
\ No newline at end of file
diff --git a/Engine/Source/Runtime/XmlParser/Private/XmlParserPrivatePCH.h b/Engine/Source/Runtime/XmlParser/Private/XmlParserPrivatePCH.h
new file mode 100644
--- /dev/null
+++ b/Engine/Source/Runtime/XmlParser/Private/XmlParserPrivatePCH.h
@@ -0,0 +1,6 @@
+// Copyright 1998-2014 Epic Games, Inc. All Rights Reserved.
+
+#pragma once
+
+#include "Core.h"
+#include "CoreUObject.h"
\ No newline at end of file
diff --git a/Engine/Source/Runtime/XmlParser/Public/XmlFile.h b/Engine/Source/Runtime/XmlParser/Public/XmlFile.h
new file mode 100644
--- /dev/null
+++ b/Engine/Source/Runtime/XmlParser/Public/XmlFile.h
@@ -0,0 +1,119 @@
+// Copyright 1998-2014 Epic Games, Inc. All Rights Reserved.
+
+#pragma once
+
+#ifndef __XmlFile_h__
+#define __XmlFile_h__
+
+#include "XmlNode.h"
+
+namespace EConstructMethod
+{
+	enum Type
+	{
+		ConstructFromFile,
+		ConstructFromBuffer,
+	};
+}
+
+/** Class representing an XML file */
+class XMLPARSER_API FXmlFile
+{
+public:
+
+
+public:
+
+	/** Constructs the file without a path */
+	FXmlFile() : RootNode(nullptr), bFileLoaded(false) {}
+	/** 
+	 * Constructs the file with the passed path. InFile is either treated as a filename to open, or as a text
+	 * buffer to load.
+	 * @param	InFile				The path/text to use
+	 * @param	ConstructMethod		Whether to load a file of use the string as a buffer of xml data
+	 */
+	FXmlFile(const FString& InFile, EConstructMethod::Type ConstructMethod = EConstructMethod::ConstructFromFile);
+	/** Dtor */
+	~FXmlFile() { Clear(); };
+
+private:
+
+	/** No conversion allowed */
+	FXmlFile(const FXmlFile& rhs) {}
+	/** No assignment allowed */
+	FXmlFile& operator=(const FXmlFile& rhs) { return *this; };
+
+public:
+
+	/** 
+	 * Loads the file with the passed path. Path is either treated as a filename to open, or as a text
+	 * buffer to load.
+	 * @param	Path				The path/text to use
+	 * @param	ConstructMethod		Whether to load a file of use the string as a buffer of xml data
+	 */
+	bool LoadFile(const FString& Path, EConstructMethod::Type ConstructMethod = EConstructMethod::ConstructFromFile);
+	/** Gets the last error message from the class */
+	FString GetLastError() const;
+	/** Clears the file of all internals. Note: Makes any existing pointers to FXmlNode's INVALID */
+	void Clear();
+	/** Checks to see if a file is loaded */
+	bool IsValid() const;
+
+	/** 
+	 * Returns the root node of the loaded file. nullptr if no file loaded. 
+	 * It is assumed that there will always be one and only one root node.
+	 * @return						Pointer to root node
+	 */
+	const FXmlNode* GetRootNode() const;
+
+	/** 
+	 * Returns the root node of the loaded file. nullptr if no file loaded. 
+	 * It is assumed that there will always be one and only one root node.
+	 * @return						Pointer to root node
+	 */
+	FXmlNode* GetRootNode();
+
+	/**
+	 * Write to disk, UTF-16 format only at the moment
+	 * @param	Path				File path to save to
+	 * @return						Whether writing the XML to a file succeeded
+	 */
+	bool Save(const FString& Path);
+
+private:
+
+	/** Pre-processes the file doing anything necessary (such as removing comments) */
+	void PreProcessInput(TArray<FString>& Input);
+	/** Whites of the text at the specified locations in a passed-in array */
+	void WhiteOut(TArray<FString>& Input, int32 LineStart, int32 LineEnd, int32 IndexStart, int32 IndexEnd);
+	/** Converts the passed input into a list of tokens for parsing */
+	TArray<FString> Tokenize(FString Input);
+	/** Converts the passed input into a list of tokens for parsing */
+	TArray<FString> Tokenize(TArray<FString>& Input);
+	/** 
+	 * Add an attribute to the passed-in array.
+	 * This makes the assumption that an attribute comes in as one 'token' (i.e. no spaces between tag="value").
+	 */
+	void AddAttribute(const FString& InToken, TArray<FXmlAttribute>& OutAttributes);
+	/** Recursively creates nodes from the list of tokens */
+	FXmlNode* CreateNodeRecursive(const TArray<FString>& Tokens, int32 StartIndex, int32* OutNextIndex = nullptr);
+	/** Hook next ptrs up recursively */
+	void HookUpNextPtrs(FXmlNode* Node);
+	/** Creates the internal file representation as a bunch of FXmlNode's */
+	void CreateNodes(const TArray<FString>& Tokens);
+
+private:
+
+	/** The passed-in path of the loaded file (might be absolute or relative) */
+	FString LoadedFile;
+	/** An error message generated on errors to return to the client */
+	FString ErrorMessage;
+	/** A pointer to the root node */
+	FXmlNode* RootNode;
+	/** Flag for whether a file is loaded or not */
+	bool bFileLoaded;
+	/** Flag for whether the node creation process failed or not */
+	bool bCreationFailed;
+};
+
+#endif
diff --git a/Engine/Source/Runtime/XmlParser/Public/XmlNode.h b/Engine/Source/Runtime/XmlParser/Public/XmlNode.h
new file mode 100644
--- /dev/null
+++ b/Engine/Source/Runtime/XmlParser/Public/XmlNode.h
@@ -0,0 +1,88 @@
+// Copyright 1998-2014 Epic Games, Inc. All Rights Reserved.
+
+#pragma once
+
+#ifndef __XmlNode_h__
+#define __XmlNode_h__
+
+/** Forward Decl for implementation */
+class FXmlFile;
+
+class XMLPARSER_API FXmlAttribute
+{
+public:
+	FXmlAttribute(const FString& InTag, const FString& InValue)
+		: Tag(InTag)
+		, Value(InValue)
+	{
+	}
+
+	/** Gets the tag of the attribute */
+	const FString& GetTag() const;
+
+	/** Gets the value of the attribute */
+	const FString& GetValue() const;
+
+private:
+	/** The tag string */
+	FString Tag;
+
+	/** The value string */
+	FString Value;
+};
+
+/** Xml Node representing a line in an xml file */
+class XMLPARSER_API FXmlNode
+{
+	friend class FXmlFile;
+
+private:
+
+	/** Default ctor, private for FXmlFile use only */
+	FXmlNode() : NextNode(nullptr) {}
+	/** No copy ctor allowed */
+	FXmlNode(const FXmlNode& rhs) {}
+	/** dtor */
+	~FXmlNode() { Delete(); }
+
+private:
+
+	/** Recursively deletes the nodes for cleanup */
+	void Delete();
+
+public:
+
+	/** Gets the next node in a list of nodes */
+	const FXmlNode* GetNextNode() const;
+	/** Gets a list of children nodes */
+	const TArray<FXmlNode*>& GetChildrenNodes() const;
+	/** Gets the first child of this node which can be iterated into with GetNextNode */
+	const FXmlNode* GetFirstChildNode() const;
+	/** Finds the first child node that contains the specified tag */
+	const FXmlNode* FindChildNode(const FString& InTag) const;
+	/** Finds the first child node that contains the specified tag */
+	FXmlNode* FindChildNode(const FString& InTag);
+	/** Gets the tag of the node */
+	const FString& GetTag() const;
+	/** Gets the value of the node */
+	const FString& GetContent() const;
+	/** Gets an attribute that corresponds with the passed-in tag */
+	FString GetAttribute(const FString& InTag) const;
+	/** Adds a simple tag with content to the current node */
+	void AppendChildNode(const FString& InTag, const FString& InContent);
+private:
+
+	/** The list of children nodes */
+	TArray<FXmlNode*> Children;
+	/** Attributes of this node */
+	TArray<FXmlAttribute> Attributes;
+	/** Tag of the node */
+	FString Tag;
+	/** Content of the node */
+	FString Content;
+	/** Next pointer */
+	FXmlNode* NextNode;
+
+};
+
+#endif
diff --git a/Engine/Source/Runtime/XmlParser/Public/XmlParser.h b/Engine/Source/Runtime/XmlParser/Public/XmlParser.h
new file mode 100644
--- /dev/null
+++ b/Engine/Source/Runtime/XmlParser/Public/XmlParser.h
@@ -0,0 +1,11 @@
+// Copyright 1998-2014 Epic Games, Inc. All Rights Reserved.
+
+#pragma once
+
+#ifndef __XmlParser_h__
+#define __XmlParser_h__
+
+#include "XmlNode.h"
+#include "XmlFile.h"
+
+#endif
diff --git a/Engine/Source/Runtime/XmlParser/XmlParser.Build.cs b/Engine/Source/Runtime/XmlParser/XmlParser.Build.cs
new file mode 100644
--- /dev/null
+++ b/Engine/Source/Runtime/XmlParser/XmlParser.Build.cs
@@ -0,0 +1,18 @@
+﻿// Copyright 1998-2014 Epic Games, Inc. All Rights Reserved.
+
+using UnrealBuildTool;
+
+public class XmlParser : ModuleRules
+{
+	public XmlParser( TargetInfo Target )
+	{
+		PublicIncludePaths.AddRange(new string[] { "Editor/XmlParser/Public" });
+		PrivateDependencyModuleNames.AddRange(
+			new string[] 
+			{ 
+				"Core",
+				"CoreUObject",
+			}
+			);
+	}
+}
