XML Parser: Allow anonymous discovery of node attributes
- Added new GetAttributes() method to FXmlNode that returns all of the attributes in the node

diff --git a/Engine/Source/Runtime/XmlParser/Public/XmlNode.h b/Engine/Source/Runtime/XmlParser/Public/XmlNode.h
--- a/Engine/Source/Runtime/XmlParser/Public/XmlNode.h
+++ b/Engine/Source/Runtime/XmlParser/Public/XmlNode.h
@@ -66,6 +66,17 @@ public:
 	const FString& GetTag() const;
 	/** Gets the value of the node */
 	const FString& GetContent() const;
+
+	/**
+	 * Gets all of the attributes in this node
+	 *
+	 * @return	List of attributes in this node
+	 */
+	const TArray<FXmlAttribute>& GetAttributes() const
+	{
+		return Attributes;
+	}
+
 	/** Gets an attribute that corresponds with the passed-in tag */
 	FString GetAttribute(const FString& InTag) const;
 	/** Adds a simple tag with content to the current node */
