XmlParser - Added method SetContent to the FXmlNode

diff --git a/Engine/Source/Runtime/XmlParser/Private/XmlNode.cpp b/Engine/Source/Runtime/XmlParser/Private/XmlNode.cpp
--- a/Engine/Source/Runtime/XmlParser/Private/XmlNode.cpp
+++ b/Engine/Source/Runtime/XmlParser/Private/XmlNode.cpp
@@ -77,6 +77,12 @@ const FString& FXmlNode::GetContent() const
 	return Content;
 }
 
+
+void FXmlNode::SetContent( const FString& InContent )
+{
+	Content = InContent;
+}
+
 FString FXmlNode::GetAttribute(const FString& InTag) const
 {
 	for(auto Iter(Attributes.CreateConstIterator()); Iter; Iter++)
diff --git a/Engine/Source/Runtime/XmlParser/Public/XmlNode.h b/Engine/Source/Runtime/XmlParser/Public/XmlNode.h
--- a/Engine/Source/Runtime/XmlParser/Public/XmlNode.h
+++ b/Engine/Source/Runtime/XmlParser/Public/XmlNode.h
@@ -66,6 +66,8 @@ public:
 	const FString& GetTag() const;
 	/** Gets the value of the node */
 	const FString& GetContent() const;
+	/** Sets the new value of the node */
+	void SetContent(const FString& InContent);
 
 	/**
 	 * Gets all of the attributes in this node
