Fix parsing of XML attributes whose values contain '=' characters, and translate named entities correctly.

diff --git a/Engine/Source/Runtime/XmlParser/Private/XmlFile.cpp b/Engine/Source/Runtime/XmlParser/Private/XmlFile.cpp
--- a/Engine/Source/Runtime/XmlParser/Private/XmlFile.cpp
+++ b/Engine/Source/Runtime/XmlParser/Private/XmlFile.cpp
@@ -607,18 +607,23 @@ static bool IsTagOperator(FString ToCheck)
 
 void FXmlFile::AddAttribute(const FString& InToken, TArray<FXmlAttribute>& OutAttributes)
 {
-	TArray<FString> Tokens;
-	if(InToken.Contains(TEXT("=")))
+	int32 EqualsIdx;
+	if(InToken.FindChar(TEXT('='), EqualsIdx))
 	{
-		InToken.ParseIntoArray(Tokens, TEXT("="), true);
-		if(Tokens.Num() == 2)
+		bool bQuotesRemoved = false;
+		FString Value = InToken.Mid(EqualsIdx + 1).TrimQuotes(&bQuotesRemoved);
+		if(bQuotesRemoved)
 		{
-			bool bQuotesRemoved = false;
-			FString Value = Tokens[1].TrimQuotes(&bQuotesRemoved);
-			if(bQuotesRemoved)
+			int32 AmpIdx;
+			if(Value.FindChar(TEXT('&'), AmpIdx))
 			{
-				OutAttributes.Add(FXmlAttribute(Tokens[0], Value));
+				Value.ReplaceInline(TEXT("&quot;"), TEXT("\""), ESearchCase::CaseSensitive);
+				Value.ReplaceInline(TEXT("&amp;"), TEXT("&"), ESearchCase::CaseSensitive);
+				Value.ReplaceInline(TEXT("&apos;"), TEXT("'"), ESearchCase::CaseSensitive);
+				Value.ReplaceInline(TEXT("&lt;"), TEXT("<"), ESearchCase::CaseSensitive);
+				Value.ReplaceInline(TEXT("&gt;"), TEXT(">"), ESearchCase::CaseSensitive);
 			}
+			OutAttributes.Add(FXmlAttribute(InToken.Left(EqualsIdx), Value));
 		}
 	}
 }
