Fast Xml parser for UE4

This is a simple XML parser that can load XML files very quickly.  The new FFastXml class has a single static method ParseXmlFile() for loading XML files.  The parser is designed to load files efficiently, but may not support all XML features or be resilient to malformed data.  Unlike the FXmlFile class, FFastXml does not generate an XML node tree.  Instead, you must supply a callback class using IFastXmlCallback and respond to elements and attributes as they are parsed.  Call the ParseXmlFile() function and pass in either the full path to the XML file to parse (XmlFilePath parameter), or load the file yourself and pass the contents using the XmlFileContents parameter.  One of either XmlFilePath or XmlFileContents must be valid for ParseXmlFile() to work.

To process the XML data, create your own class that implements the IFastXmlCallback interface.  This class will be responsible for processing the XML elements as they are loaded by FFastXml::ParseXmlFile().  You'll receive a ProcessElement() call for every XML element that is encountered, along with a corresponding ProcessClose() when that element's scope has ended.  ProcessAttribute() will be called for any attributes found within the scope of the current element.

Remember to add a module dependency on "XmlParser" in order to be able to call FFastXml::ParserXmlFile().

This XML parser based on John W. Ratcliff's "FastXml". (FastXml.tps included, with more info.)

diff --git a/Engine/Source/Runtime/XmlParser/Private/FastXml.cpp b/Engine/Source/Runtime/XmlParser/Private/FastXml.cpp
new file mode 100644
--- /dev/null
+++ b/Engine/Source/Runtime/XmlParser/Private/FastXml.cpp
@@ -0,0 +1,698 @@
+#include "XmlParserPrivatePCH.h"
+#include "FastXml.h"
+
+#define LOCTEXT_NAMESPACE "FastXml"
+
+
+/**
+ * Implementation of fast XML parser
+ */
+class FFastXmlImpl
+{
+
+public:
+
+	FFastXmlImpl( IFastXmlCallback* InitCallback, TCHAR* InitXmlFileContents, const SIZE_T InitXmlFileContentsLength, FFeedbackContext* InitFeedbackContext )
+		: Callback( InitCallback ),
+		  XmlFileContents( InitXmlFileContents ),
+		  XmlFileContentsLength( InitXmlFileContentsLength ),
+		  FeedbackContext( InitFeedbackContext )
+	{
+		// Setup character type map.  It maps each ASCII character to the type of XML character it is.
+		{
+			memset( CharacterTypeMap, (int)ECharType::Data, sizeof( CharacterTypeMap ) );
+			CharacterTypeMap[ 0 ] = ECharType::EndOfFile;
+			CharacterTypeMap[ L' ' ] = CharacterTypeMap[ L'\t' ] = ECharType::Whitespace;
+			CharacterTypeMap[ L'/' ] = CharacterTypeMap[ L'>' ] = CharacterTypeMap[ L'?' ] = ECharType::EndOfElement;
+			CharacterTypeMap[ L'\n' ] = CharacterTypeMap[ L'\r' ] = ECharType::EndOfLine;
+		}
+
+		StackIndex = 0;
+		for( uint32 StackLevel = 0; StackLevel < ( MaxStackDepth + 1 ); StackLevel++ )
+		{
+			Stack[ StackLevel ] = nullptr;
+			IsStackAllocated[ StackLevel ] = false;
+		}
+	}
+
+
+	~FFastXmlImpl( void )
+	{
+		StackIndex = 0;
+
+		for( uint32 StackLevel = 0; StackLevel < ( StackIndex + 1 ); StackLevel++ )
+		{
+			if( IsStackAllocated[ StackLevel ] )
+			{
+				FMemory::Free( (void*)Stack[ StackLevel ] );
+				IsStackAllocated[ StackLevel ] = false;
+			}
+			Stack[ StackLevel ] = nullptr;
+		}
+	}
+
+
+	bool ProcessXmlFile( FText& OutErrorMessage, int32& OutErrorLineNumber )
+	{
+		// Parse the file!
+		bool bResult = ProcessXmlFileInternal() && ( FeedbackContext == nullptr || !FeedbackContext->ReceivedUserCancel() );
+
+		// Keep track of any errors for the caller
+		OutErrorMessage = ErrorMessage;
+		OutErrorLineNumber = LineNumber;
+
+		return bResult;
+	}
+
+
+protected:
+
+
+	TCHAR* ProcessClose( const TCHAR Char, const TCHAR* Element, TCHAR* Buffer, const int32 AttributeCount, const TCHAR** Attributes )
+	{
+		if( FeedbackContext != nullptr )
+		{
+			FeedbackContext->UpdateProgress( Buffer - &XmlFileContents[ 0 ], XmlFileContentsLength );
+		}
+
+		if( Char == L'/' || Char == L'?' )
+		{
+			TCHAR* Slash = (TCHAR*)FCString::Strchr( Element, Char );
+			if( Slash )
+			{
+				*Slash = 0;
+			}
+
+			if( Char == L'?' && FCString::Strcmp( Element, TEXT( "xml" ) ) == 0 )
+			{
+				if( !Callback->ProcessXmlDeclaration( 0, LineNumber ) )
+				{
+					return nullptr;
+				}
+			}
+			else
+			{
+				if( !Callback->ProcessElement( Element, 0, LineNumber ) )
+				{
+					ErrorMessage = LOCTEXT( "UserAbortedOnElement", "User aborted the parsing process" );
+					return nullptr;
+				}
+
+				check( AttributeCount % 2 == 0 );
+				for( int32 AttributeIndex = 0; AttributeIndex < AttributeCount / 2; ++AttributeIndex )
+				{
+					if( !Callback->ProcessAttribute( Attributes[ AttributeIndex * 2 ], Attributes[ AttributeIndex * 2 + 1 ] ) )
+					{
+						ErrorMessage = LOCTEXT( "UserAbortedOnAttribute", "User aborted the parsing process" );
+						return nullptr;
+					}
+				}
+
+				PushElement( Element );
+
+				const TCHAR* Close = PopElement();
+				if( !Callback->ProcessClose( Close ) )
+				{
+					return nullptr;
+				}
+			}
+
+			if( !Slash )
+			{
+				++Buffer;
+			}
+		}
+		else
+		{
+			Buffer = SkipNextData( Buffer );
+			TCHAR* Data = Buffer; // this is the data portion of the element, only copies memory if we encounter line feeds
+			TCHAR* DestData = 0;
+			while( *Buffer && *Buffer != L'<' )
+			{
+				if( CharacterTypeMap[ *Buffer ] == ECharType::EndOfLine )
+				{
+					if( *Buffer == L'\r' )
+					{
+						LineNumber++;
+					}
+					DestData = Buffer;
+					*DestData++ = L' '; // replace the linefeed with a space...
+					Buffer = SkipNextData( Buffer );
+					while( *Buffer && *Buffer != L'<' )
+					{
+						if( CharacterTypeMap[ *Buffer ] == ECharType::EndOfLine )
+						{
+							if( *Buffer == L'\r' )
+							{
+								LineNumber++;
+							}
+							*DestData++ = L' '; // replace the linefeed with a space...
+							Buffer = SkipNextData( Buffer );
+						}
+						else
+						{
+							*DestData++ = *Buffer++;
+						}
+					}
+					break;
+				}
+				else
+				{
+					++Buffer;
+				}
+			}
+
+			if( *Buffer == L'<' )
+			{
+				if( DestData )
+				{
+					*DestData = 0;
+				}
+				else
+				{
+					*Buffer = 0;
+				}
+
+				Buffer++; // skip it..
+
+				if( *Data == 0 )
+				{
+					Data = 0;
+				}
+
+				if( !Callback->ProcessElement( Element, Data, LineNumber ) )
+				{
+					ErrorMessage = LOCTEXT( "UserAbortedOnElement", "User aborted the parsing process" );
+					return nullptr;
+				}
+
+				check( AttributeCount % 2 == 0 );
+				for( int32 AttributeIndex = 0; AttributeIndex < AttributeCount / 2; ++AttributeIndex )
+				{
+					if( !Callback->ProcessAttribute( Attributes[ AttributeIndex * 2 ], Attributes[ AttributeIndex * 2 + 1 ] ) )
+					{
+						ErrorMessage = LOCTEXT( "UserAbortedOnAttribute", "User aborted the parsing process" );
+						return nullptr;
+					}
+				}
+
+				PushElement( Element );
+
+				// check for the comment use case...
+				if( Buffer[ 0 ] == L'!' && Buffer[ 1 ] == L'-' && Buffer[ 2 ] == L'-' )
+				{
+					Buffer += 3;
+					while( *Buffer && *Buffer == L' ' )
+					{
+						++Buffer;
+					}
+
+					TCHAR* Comment = Buffer;
+					TCHAR* CommentEnd = FCString::Strstr( Buffer, TEXT( "-->" ) );
+					if( CommentEnd )
+					{
+						*CommentEnd = 0;
+						Buffer = CommentEnd + 3;
+						if( !Callback->ProcessComment( Comment ) )
+						{
+							ErrorMessage = LOCTEXT( "UserAbortedOnComment", "User aborted the parsing process" );
+							return nullptr;
+						}
+					}
+				}
+				else if( *Buffer == L'/' )
+				{
+					Buffer = ProcessClose( Buffer );
+				}
+			}
+			else
+			{
+				ErrorMessage = LOCTEXT( "ElementDataNotTerminated", "Data portion of an element wasn't terminated properly" );
+				return nullptr;
+			}
+		}
+
+		if( FeedbackContext != nullptr && FeedbackContext->ReceivedUserCancel() )
+		{
+			Buffer = nullptr;
+			ErrorMessage = LOCTEXT( "UserAbortedOnFile", "User cancelled processing of this file" );
+		}
+
+		return Buffer;
+	}
+
+
+	TCHAR* ProcessClose( TCHAR* Buffer )
+	{
+		const TCHAR* Start = PopElement();
+		const TCHAR* Close = Start;
+		if( Buffer[ 1 ] != L'>' )
+		{
+			Buffer++;
+			Close = Buffer;
+			
+			while( *Buffer && *Buffer != L'>' )
+			{
+				Buffer++;
+			}
+
+			*Buffer = 0;
+		}
+
+		if( 0 != FCString::Strcmp( Start, Close ) )
+		{
+			ErrorMessage = LOCTEXT( "OpenCloseTagsNotMatched", "Open and closing tags do not match" );
+			return nullptr;
+		}
+
+		if( !Callback->ProcessClose( Close ) )
+		{
+			return nullptr;
+		}
+		++Buffer;
+
+		return Buffer;
+	}
+
+
+	bool ProcessXmlFileInternal()
+	{
+		LineNumber = 1;
+
+		for( uint32 StackLevel = 0; StackLevel < ( StackIndex + 1 ); StackLevel++ )
+		{
+			if( !IsStackAllocated[ StackLevel ] )
+			{
+				const TCHAR* Text = Stack[ StackLevel ];
+				if( Text )
+				{
+					uint32 TextLength = (uint32)FCString::Strlen( Text );
+					Stack[ StackLevel ] = (const TCHAR*)FMemory::Malloc( TextLength * sizeof( TCHAR ) + sizeof( TCHAR ) );
+					memcpy( (void*)Stack[ StackLevel ], Text, TextLength * sizeof( TCHAR ) + sizeof( TCHAR ) );
+					IsStackAllocated[ StackLevel ] = true;
+				}
+			}
+		}
+
+		TCHAR* Element;
+		TCHAR* Buffer = &XmlFileContents[ 0 ];
+
+		while( *Buffer )
+		{
+			Buffer = SkipNextData( Buffer );
+			if( *Buffer == 0 )
+			{
+				break;
+			}
+
+			if( *Buffer == L'<' )
+			{
+				Buffer++;
+				if( *Buffer == L'?' ) // Allow XML declarations
+				{
+					Buffer++;
+				}
+				else if( Buffer[ 0 ] == L'!' && Buffer[ 1 ] == L'-' && Buffer[ 2 ] == L'-' )
+				{
+					Buffer += 3;
+					while( *Buffer && *Buffer == L' ' )
+					{
+						Buffer++;
+					}
+
+					TCHAR* Comment = Buffer;
+					TCHAR* CommentEnd = FCString::Strstr( Buffer, TEXT( "-->" ) );
+					if( CommentEnd )
+					{
+						*CommentEnd = 0;
+						Buffer = CommentEnd + 3;
+						if( !Callback->ProcessComment( Comment ) )
+						{
+							ErrorMessage = LOCTEXT( "UserAbortedOnComment", "User aborted the parsing process" );
+							return false;
+						}
+					}
+					continue;
+				}
+			}
+
+			if( *Buffer == L'/' )
+			{
+				Buffer = ProcessClose( Buffer );
+				if( !Buffer )
+				{
+					return false;
+				}
+			}
+			else
+			{
+				if( *Buffer == L'?' )
+				{
+					Buffer++;
+				}
+
+				Element = Buffer;
+				int32 AttributeCount = 0;
+				const TCHAR* Attributes[ MaxAttributes ];
+				bool Close;
+				Buffer = NextWhitespaceOrClose( Buffer, Close );
+				if( Close )
+				{
+					TCHAR Char = *( Buffer - 1 );
+					if( Char != L'?' && Char != L'/' )
+					{
+						Char = L'>';
+					}
+					*Buffer++ = 0;
+					Buffer = ProcessClose( Char, Element, Buffer, AttributeCount, Attributes );
+					if( !Buffer )
+					{
+						return false;
+					}
+				}
+				else
+				{
+					if( *Buffer == 0 )
+					{
+						return true;
+					}
+
+					*Buffer = 0; // place a zero byte to indicate the end of the element name...
+					Buffer++;
+
+					while( *Buffer )
+					{
+						Buffer = SkipNextData( Buffer ); // advance past any soft seperators (tab or space)
+
+						if( CharacterTypeMap[ *Buffer ] == ECharType::EndOfElement )
+						{
+							TCHAR Char = *Buffer++;
+							if( L'?' == Char )
+							{
+								if( L'>' != *Buffer ) //?>
+								{
+									check( 0 );
+									return false;
+								}
+
+								Buffer++;
+							}
+							Buffer = ProcessClose( Char, Element, Buffer, AttributeCount, Attributes );
+							if( !Buffer )
+							{
+								return false;
+							}
+							break;
+						}
+						else
+						{
+							if( AttributeCount >= MaxAttributes )
+							{
+								ErrorMessage = LOCTEXT( "TooManyAttributes", "Encountered too many attributes in a single element for this parser to handle" );
+								return false;
+							}
+							Attributes[ AttributeCount ] = Buffer;
+							Buffer = NextSeparator( Buffer );  // scan up to a space, or an equal
+							if( *Buffer )
+							{
+								if( *Buffer != L'=' )
+								{
+									*Buffer = 0;
+									Buffer++;
+									while( *Buffer && *Buffer != L'=' )
+									{
+										Buffer++;
+									}
+
+									if( *Buffer == L'=' )
+									{
+										Buffer++;
+									}
+								}
+								else
+								{
+									*Buffer = 0;
+									Buffer++;
+								}
+
+								if( *Buffer ) // if not eof...
+								{
+									Buffer = SkipNextData( Buffer );
+									if( *Buffer == L'"' )
+									{
+										Buffer++;
+										AttributeCount++;
+										Attributes[ AttributeCount ] = Buffer;
+										AttributeCount++;
+										while( *Buffer && *Buffer != 34 )
+										{
+											Buffer++;
+										}
+										if( *Buffer == L'"' )
+										{
+											*Buffer = 0;
+											Buffer++;
+										}
+										else
+										{
+											ErrorMessage = LOCTEXT( "NoClosingQuoteForAttribute", "Failed to find closing quote for attribute" );
+											return false;
+										}
+									}
+									else
+									{
+										// Missing quote after attribute.  We'll handle it as best we can.
+										AttributeCount--;
+										while( *Buffer != L'/' && *Buffer != L'>' && *Buffer != 0 )
+										{
+											Buffer++;
+										}
+									}
+								}
+							}
+						}
+					}
+				}
+			}
+		}
+
+		if( StackIndex )
+		{
+			ErrorMessage = LOCTEXT( "InvalidFileFormat", "Invalid file format" );
+			return false;
+		}
+
+		return true;
+	}
+
+	
+	inline TCHAR* NextWhitespace( TCHAR* Buffer )
+	{
+		while( *Buffer && CharacterTypeMap[ *Buffer ] != ECharType::Whitespace )
+		{
+			Buffer++;
+		}
+
+		return Buffer;
+	}
+
+
+	inline TCHAR* NextWhitespaceOrClose( TCHAR* Buffer, bool& IsClose )
+	{
+		while( *Buffer && CharacterTypeMap[ *Buffer ] != ECharType::Whitespace && *Buffer != L'>' )
+		{
+			Buffer++;
+		}
+
+		IsClose = *Buffer == L'>';
+		return Buffer;
+	}
+
+
+	inline TCHAR* NextSeparator( TCHAR* Buffer )
+	{
+		while( *Buffer && CharacterTypeMap[ *Buffer ] != ECharType::Whitespace && *Buffer != L'=' )
+		{
+			Buffer++;
+		}
+
+		return Buffer;
+	}
+
+
+	inline TCHAR* SkipNextData( TCHAR* Buffer )
+	{
+		// while we have data, and we encounter soft seperators or line feeds...
+		while( *Buffer && CharacterTypeMap[ *Buffer ] == ECharType::Whitespace || CharacterTypeMap[ *Buffer ] == ECharType::EndOfLine )
+		{
+			if( *Buffer == L'\n' )
+			{
+				LineNumber++;
+			}
+
+			Buffer++;
+		}
+		return Buffer;
+	}
+
+
+	void PushElement( const TCHAR* Element )
+	{
+		check( StackIndex < MaxStackDepth );
+		if( StackIndex < MaxStackDepth )
+		{
+			if( IsStackAllocated[ StackIndex ] )
+			{
+				FMemory::Free( (void*)Stack[ StackIndex ] );
+			}
+			IsStackAllocated[ StackIndex ] = false;
+			Stack[ StackIndex++ ] = Element;
+		}
+	}
+
+	const TCHAR* PopElement( void )
+	{
+		if( IsStackAllocated[ StackIndex ] )
+		{
+			FMemory::Free( (void*)Stack[ StackIndex ] );
+			Stack[ StackIndex ] = nullptr;
+			IsStackAllocated[ StackIndex ] = false;
+		}
+		return StackIndex ? Stack[ --StackIndex ] : nullptr;
+	}
+
+
+	/** User callback to report XML data and progress */
+	IFastXmlCallback* Callback;
+
+	/** Reference to the the contents of the XML file to parse */
+	TCHAR* XmlFileContents;
+
+	/** Length of the XmlFileContents string */
+	SIZE_T XmlFileContentsLength;
+
+	/** Feedback context for status reporting.  Can be nullptr. */
+	FFeedbackContext* FeedbackContext;
+
+	/** If anything goes wrong or the user cancels, the error message to return will be stored here */
+	FText ErrorMessage;
+
+	/** Types of characters we'll encounter while parsing */
+	enum class ECharType : uint8
+	{
+		Data = 0,
+		EndOfFile,
+		Whitespace,
+		EndOfElement, // Either a forward slash or a greater than symbol
+		EndOfLine,
+	};
+
+	/** Maps each ASCII character to the type of character we think it is */
+	ECharType CharacterTypeMap[ 256 ];
+
+	/** Current stack depth as we descend through XML nodes and attributes */
+	uint32 StackIndex;
+
+	/** The current line number we're on in the file */
+	int32 LineNumber;
+
+	/** Maximum stack depth.  We can't support XML documents that go deeper than this! */
+	static const int MaxStackDepth = 2048;
+
+	/** Maximum number of attributes in an element that we can support */
+	static const int MaxAttributes = 2048;
+
+	/** Stack of characters, one for each stack depth */
+	const TCHAR* Stack[ MaxStackDepth + 1 ];
+
+	/** For each stack level, whether we've allocated that stack or not yet */
+	bool IsStackAllocated[ MaxStackDepth + 1 ];
+};
+
+
+
+bool FFastXml::ParseXmlFile( IFastXmlCallback* Callback, const TCHAR* XmlFilePath, TCHAR* XmlFileContents, FFeedbackContext* FeedbackContext, const bool bShowSlowTaskDialog, const bool bShowCancelButton, FText& OutErrorMessage, int32& OutErrorLineNumber )
+{
+	bool bSuccess = true;
+
+	FString LoadedXmlFileContents;
+	SIZE_T XmlFileContentsLength = 0;
+	if( XmlFilePath != nullptr && FPlatformString::Strlen( XmlFilePath ) > 0 )
+	{
+		if( FPlatformFileManager::Get().GetPlatformFile().FileExists( XmlFilePath ) )
+		{
+			if( FeedbackContext != nullptr )
+			{
+				FeedbackContext->BeginSlowTask( LOCTEXT( "LoadingXML", "Loading XML file..." ), bShowSlowTaskDialog, false /* Cannot support cancelling the loading part */ );
+			}
+
+			if( FFileHelper::LoadFileToString( LoadedXmlFileContents, XmlFilePath ) )
+			{
+				XmlFileContentsLength = LoadedXmlFileContents.Len();
+				if( XmlFileContentsLength > 0 )
+				{
+					// File was loaded okay!
+					XmlFileContents = &LoadedXmlFileContents[ 0 ];
+				}
+				else
+				{ 
+					bSuccess = false;
+					OutErrorMessage = LOCTEXT( "LoadedXMLFileWasEmpty", "The XML file is empty" );
+					OutErrorLineNumber = 1;
+				}
+			}
+			else
+			{
+				bSuccess = false;
+				OutErrorMessage = LOCTEXT( "ErrorReadingFile", "Unable to load the XML file" );
+				OutErrorLineNumber = 1;
+			}
+
+			if( FeedbackContext != nullptr )
+			{
+				FeedbackContext->EndSlowTask();
+			}
+		}
+		else
+		{
+			bSuccess = false;
+			OutErrorMessage = LOCTEXT( "FileNotFound", "Couldn't find the specified XML file on disk" );
+			OutErrorLineNumber = 1;
+		}
+	}
+	else
+	{
+		if( ensure( XmlFileContents != nullptr && *XmlFileContents != 0 ) )
+		{
+			XmlFileContentsLength = FPlatformString::Strlen( XmlFileContents );
+		}
+		else
+		{
+			bSuccess = false;
+			OutErrorMessage = LOCTEXT( "NoFileNameOrContentsPassedIn", "ParseXmlFile() was called without either an XML file name or an XML file contents text buffer supplied.  Either XmlFilePath or XmlFileContents must be valid in order to call ParseXmlFile()" );
+			OutErrorLineNumber = 1;
+		}
+	}
+
+	if( bSuccess )
+	{
+		if( FeedbackContext != nullptr )
+		{
+			FeedbackContext->BeginSlowTask( LOCTEXT( "ProcessingXML", "Processing XML file..." ), bShowSlowTaskDialog, bShowCancelButton );
+		}
+
+		// Parse the XML file contents!
+		TUniquePtr< FFastXmlImpl > FastXmlImpl( new FFastXmlImpl( Callback, XmlFileContents, XmlFileContentsLength, FeedbackContext ) );
+		bSuccess = FastXmlImpl->ProcessXmlFile( OutErrorMessage, OutErrorLineNumber );
+
+		if( FeedbackContext != nullptr )
+		{
+			FeedbackContext->EndSlowTask();
+		}
+	}
+
+	return bSuccess;
+}
+
+
+#undef LOCTEXT_NAMESPACE
\ No newline at end of file
diff --git a/Engine/Source/Runtime/XmlParser/Public/FastXml.h b/Engine/Source/Runtime/XmlParser/Public/FastXml.h
new file mode 100644
--- /dev/null
+++ b/Engine/Source/Runtime/XmlParser/Public/FastXml.h
@@ -0,0 +1,103 @@
+#pragma once
+
+
+/**
+ * Fast XML parser based on John W. Ratcliff's "FastXml" (see FastXml.tps and http://codesuppository.blogspot.com/2009/02/fastxml-extremely-lightweight-stream.html)
+ * 
+ * This is a simple XML parser that can load XML files very quickly.  The class has a single static method, ParseXmlFile() for
+ * loading XML files.  The parser is designed to load files efficiently, but may not support all XML features or be resilient 
+ * to malformed data.  Unlike the FXmlFile class, FFastXml does not generate an XML node tree.  Instead, you must supply a 
+ * callback class using IFastXmlCallback and respond to elements and attributes as they are parsed.  Call the ParseXmlFile() 
+ * function and pass in either the full path to the XML file to parse (XmlFilePath parameter), or load the file yourself and 
+ * pass the contents using the XmlFileContents parameter.  One of either XmlFilePath or XmlFileContents must be valid for 
+ * ParseXmlFile() to work.
+ * 
+ * Remember to add a module dependency on "XmlParser" in order to be able to call FFastXml::ParserXmlFile().
+ */
+class FFastXml
+{
+
+public:
+
+	/**
+	 * Quickly parse an XML file.  Pass in your implementation of the IFastXmlCallback interface with code to handle parsed 
+	 * elements and attributes, along with either the full path to the XML file to load, or a string with the full XML file content.
+	 * 
+	 * @param	Callback			As the parser encounters XML elements or attributes, methods on this callback object will be called
+	 * @param	XmlFilePath			The path on disk to the XML file to load, or an empty string if you'll be passing the XML file content in directly using the XmlFileContents parameter.
+	 * @param	XmlFileContents		The full contents of the file to parse, or an empty string if you've passed a full path to the file to parse using the XmlFilePath parameter.  Note that this string will be modified during the parsing process, so make a copy of it first if you need it for something else (uncommon.)
+	 * @param	FeedbackContext		Optional feedback context for reporting warnings or progress (GWarn is typically passed in.)  You can pass nullptr if you don't want any progress reported.
+	 * @param	bShowSlowTaskDialog	True if we should display a 'please wait' dialog while we parse the file, if the feedback context supports that
+	 * @param	bShowCancelButton	Whether the user is allowed to cancel the load of this XML file while in progress.  If the user cancels loading, an appropriate error message will be returned.  Only applies when bShowSlowTaskDialog is also enabled, and for feedback contexts that support cancel buttons.
+	 * @param	OutErrorMessage		If anything went wrong or the user canceled the load, this error will contain a description of the problem
+	 * @param	OutErrorLineNumber	Line number that any error happened on
+	 *
+	 * @return	Returns true if loading was successful, or false if anything went wrong or the user canceled.  When false is returned, an error message string will be set in the OutErrorMessage variable, along with the line number that parsing failed at in OutErrorLineNumber.
+	 */
+	XMLPARSER_API static bool ParseXmlFile( class IFastXmlCallback* Callback, const TCHAR* XmlFilePath, TCHAR* XmlFileContents, class FFeedbackContext* FeedbackContext, const bool bShowSlowTaskDialog, const bool bShowCancelButton, FText& OutErrorMessage, int32& OutErrorLineNumber );
+	
+};
+	
+
+/**
+ * Create your own class that implements the IFastXmlCallback interface to process the XML elements as they
+ * are loaded by FFastXml::ParseXmlFile().  You'll receive a ProcessElement() call for every XML element
+ * that is encountered, along with a corresponding ProcessClose() when that element's scope has ended.
+ * ProcessAttribute() will be called for any attributes found within the scope of the current element.
+ */
+class IFastXmlCallback
+{
+
+public:
+
+	/**
+	 * Called after the XML's header is parsed.  This is usually the first call that you'll get back.
+	 *
+	 * @param	ElementData			Optional data for this element, nullptr if none
+	 * @param	XmlFileLineNumber	Line number in the XML file we're on
+	 *
+	 * @return	You should return true to continue processing the file, or false to stop processing immediately.
+	 */
+	virtual bool ProcessXmlDeclaration( const TCHAR* ElementData, int32 XmlFileLineNumber ) = 0;
+
+	/**
+	 * Called when a new XML element is encountered, starting a new scope.  You'll receive a call to ProcessClose()
+	 * when this element's scope has ended.
+	 *
+	 * @param	ElementName			The name of the element
+	 * @param	ElementData			Optional data for this element, nullptr if none
+	 * @param	XmlFileLineNumber	The line number in the XML file we're on
+	 *
+	 * @return	You should return true to continue processing the file, or false to stop processing immediately.
+	 */
+	virtual bool ProcessElement( const TCHAR* ElementName, const TCHAR* ElementData, int32 XmlFileLineNumber ) = 0;
+	
+	/**
+	 * Called when an XML attribute is encountered for the current scope's element.
+	 *
+	 * @param	AttributeName	The name of the attribute
+	 * @param	AttributeValue	The value of the attribute
+	 *
+	 * @return	You should return true to continue processing the file, or false to stop processing immediately.
+	 */
+	virtual bool ProcessAttribute( const TCHAR* AttributeName, const TCHAR* AttributeValue ) = 0;
+
+	/**
+	 * Called when an element's scope ends in the XML file
+	 *
+	 * @param	ElementName		Name of the element whose scope closed
+	 *
+	 * @return	You should return true to continue processing the file, or false to stop processing immediately.
+	 */
+	virtual bool ProcessClose( const TCHAR* Element ) = 0;
+	
+	/**
+	 * Called when a comment is encountered.  This can happen pretty much anywhere in the file.
+	 *
+	 * @param	Comment		The comment text
+	 */
+	virtual bool ProcessComment( const TCHAR* Comment ) = 0;
+	
+};
+
+
