Fixed CIS error on mac with FastXml

diff --git a/Engine/Source/Runtime/XmlParser/Private/FastXml.cpp b/Engine/Source/Runtime/XmlParser/Private/FastXml.cpp
--- a/Engine/Source/Runtime/XmlParser/Private/FastXml.cpp
+++ b/Engine/Source/Runtime/XmlParser/Private/FastXml.cpp
@@ -523,7 +523,7 @@ protected:
 	inline TCHAR* SkipNextData( TCHAR* Buffer )
 	{
 		// while we have data, and we encounter soft seperators or line feeds...
-		while( *Buffer && CharacterTypeMap[ *Buffer ] == ECharType::Whitespace || CharacterTypeMap[ *Buffer ] == ECharType::EndOfLine )
+		while( *Buffer && ( CharacterTypeMap[ *Buffer ] == ECharType::Whitespace || CharacterTypeMap[ *Buffer ] == ECharType::EndOfLine ) )
 		{
 			if( *Buffer == L'\n' )
 			{
