Add attribute serialization to XmlFile writer.

diff --git a/Engine/Source/Runtime/XmlParser/Private/XmlFile.cpp b/Engine/Source/Runtime/XmlParser/Private/XmlFile.cpp
--- a/Engine/Source/Runtime/XmlParser/Private/XmlFile.cpp
+++ b/Engine/Source/Runtime/XmlParser/Private/XmlFile.cpp
@@ -152,38 +152,12 @@ FXmlNode* FXmlFile::GetRootNode()
 
 bool FXmlFile::Save(const FString& Path)
 {
-	FString Xml;
-	Xml += TEXT("<?xml version=\"1.0\" encoding=\"UTF-8\"?>") LINE_TERMINATOR;
+	FString Xml = TEXT("<?xml version=\"1.0\" encoding=\"UTF-8\"?>") LINE_TERMINATOR;
 
-	TArray<const FXmlNode*> NodeStack;
-	static const TCHAR Tabs[] = TEXT("\t\t\t\t\t\t\t\t\t\t\t\t\t\t\t\t");
-	auto Indent = &Tabs[ARRAY_COUNT(Tabs) - 1];
-	const auto* CurrentNode = GetRootNode();
-	while (CurrentNode)
+	const FXmlNode* CurrentNode = GetRootNode();
+	if(CurrentNode != nullptr)
 	{
-		if (auto Child = CurrentNode->GetFirstChildNode())
-		{
-			NodeStack.Push(CurrentNode);
-			Xml += FString::Printf(TEXT("%s<%s>") LINE_TERMINATOR, Indent, *CurrentNode->GetTag());
-
-			CurrentNode = Child;	
-			// Actually increases the indent by one tab
-			--Indent;
-			continue;
-		}
-
-		auto Tag = *CurrentNode->GetTag();
-		Xml += FString::Printf(TEXT("%s<%s>%s</%s>") LINE_TERMINATOR, Indent, Tag, *CurrentNode->GetContent(), Tag);
-		CurrentNode = CurrentNode->GetNextNode();
-
-		while (!CurrentNode && NodeStack.Num() != 0)
-		{
-			// Actually decreases the indent by one tab
-			++Indent;
-			CurrentNode = NodeStack.Pop();
-			Xml += FString::Printf(TEXT("%s</%s>") LINE_TERMINATOR, Indent, *CurrentNode->GetTag());
-			CurrentNode = CurrentNode->GetNextNode();
-		}
+		WriteNodeHierarchy(*CurrentNode, FString(), Xml);
 	}
 	
 	TUniquePtr<FArchive> Archive(IFileManager::Get().CreateFileWriter(*Path));
@@ -892,3 +866,43 @@ void FXmlFile::CreateNodes(const TArray<FString>& Tokens)
 		ErrorMessage = NSLOCTEXT("XmlParser", "NodeCreateFail", "Failed to parse the loaded document").ToString();
 	}
 }
+
+void FXmlFile::WriteNodeHierarchy(const FXmlNode& Node, const FString& Indent, FString& Output)
+{
+	// Write the tag
+	Output += Indent + FString::Printf(TEXT("<%s"), *Node.GetTag());
+	for(const FXmlAttribute& Attribute: Node.GetAttributes())
+	{
+		FString EscapedValue = Attribute.GetValue();
+		EscapedValue.ReplaceInline(TEXT("&"), TEXT("&amp;"), ESearchCase::CaseSensitive);
+		EscapedValue.ReplaceInline(TEXT("\""), TEXT("&quot;"), ESearchCase::CaseSensitive);
+		EscapedValue.ReplaceInline(TEXT("'"), TEXT("&apos;"), ESearchCase::CaseSensitive);
+		EscapedValue.ReplaceInline(TEXT("<"), TEXT("&lt;"), ESearchCase::CaseSensitive);
+		EscapedValue.ReplaceInline(TEXT(">"), TEXT("&gt;"), ESearchCase::CaseSensitive);
+		Output += FString::Printf(TEXT(" %s=\"%s\""), *Attribute.GetTag(), *EscapedValue);
+	}
+
+	// Write the node contents
+	const FXmlNode* FirstChildNode = Node.GetFirstChildNode();
+	if(FirstChildNode == nullptr)
+	{
+		const FString& Content = Node.GetContent();
+		if(Content.Len() == 0)
+		{
+			Output += TEXT(" />") LINE_TERMINATOR;
+		}
+		else
+		{
+			Output += TEXT(">") + Content + FString::Printf(TEXT("</%s>"), *Node.GetTag()) + LINE_TERMINATOR;
+		}
+	}
+	else
+	{
+		Output += TEXT(">") LINE_TERMINATOR;
+		for(const FXmlNode* ChildNode = FirstChildNode; ChildNode != nullptr; ChildNode = ChildNode->GetNextNode())
+		{
+			WriteNodeHierarchy(*ChildNode, Indent + TEXT("\t"), Output);
+		}
+		Output += Indent + FString::Printf(TEXT("</%s>"), *Node.GetTag()) + LINE_TERMINATOR;
+	}
+}
diff --git a/Engine/Source/Runtime/XmlParser/Public/XmlFile.h b/Engine/Source/Runtime/XmlParser/Public/XmlFile.h
--- a/Engine/Source/Runtime/XmlParser/Public/XmlFile.h
+++ b/Engine/Source/Runtime/XmlParser/Public/XmlFile.h
@@ -101,6 +101,8 @@ private:
 	void HookUpNextPtrs(FXmlNode* Node);
 	/** Creates the internal file representation as a bunch of FXmlNode's */
 	void CreateNodes(const TArray<FString>& Tokens);
+	/** Writes a node hierarchy at the given root to a string */
+	static void WriteNodeHierarchy(const FXmlNode& Node, const FString& Indent, FString& Output);
 
 private:
 
