Fixed potential infinite recursion when processing malformed XML

diff --git a/Engine/Source/Runtime/XmlParser/Private/XmlFile.cpp b/Engine/Source/Runtime/XmlParser/Private/XmlFile.cpp
--- a/Engine/Source/Runtime/XmlParser/Private/XmlFile.cpp
+++ b/Engine/Source/Runtime/XmlParser/Private/XmlFile.cpp
@@ -496,6 +496,7 @@ FXmlNode* FXmlFile::CreateNodeRecursive(const TArray<FString>& Tokens, int32 Sta
 	//  - Continue parsing until </tag> for self is found
 	//  - Return own constructed node (and index of next starting point
 
+	const int32 RecursiveStartIndex = StartIndex;
 	int32 SavedIndex = StartIndex;
 
 	// Get the tag & any attributes
@@ -608,7 +609,11 @@ FXmlNode* FXmlFile::CreateNodeRecursive(const TArray<FString>& Tokens, int32 Sta
 			if(Tokens[i] == TEXT("<"))
 			{
 				// Recursively enter function creating a child at the new tag
-				FXmlNode* Child = CreateNodeRecursive(Tokens, i, &SavedIndex);
+				FXmlNode* Child = nullptr;
+				if (i > RecursiveStartIndex)
+				{
+					Child = CreateNodeRecursive(Tokens, i, &SavedIndex);
+				}
 
 				// Save child to parent
 				if(Child != nullptr)
