PR #7476: Fixed Crash FastXml (Contributed by [REDACTED_PERSON])

diff --git a/Engine/Source/Runtime/XmlParser/Private/FastXml.cpp b/Engine/Source/Runtime/XmlParser/Private/FastXml.cpp
--- a/Engine/Source/Runtime/XmlParser/Private/FastXml.cpp
+++ b/Engine/Source/Runtime/XmlParser/Private/FastXml.cpp
@@ -133,7 +133,7 @@ protected:
 			TCHAR* DestData = 0;
 			while( *Buffer && *Buffer != L'<' )
 			{
-				if( CharacterTypeMap[ *Buffer ] == ECharType::EndOfLine )
+				if( IsCharTypeEqual(*Buffer, ECharType::EndOfLine) )
 				{
 					if( *Buffer == L'\r' )
 					{
@@ -144,7 +144,7 @@ protected:
 					Buffer = SkipNextData( Buffer );
 					while( *Buffer && *Buffer != L'<' )
 					{
-						if( CharacterTypeMap[ *Buffer ] == ECharType::EndOfLine )
+						if( IsCharTypeEqual(*Buffer, ECharType::EndOfLine) )
 						{
 							if( *Buffer == L'\r' )
 							{
@@ -389,7 +389,7 @@ protected:
 					{
 						Buffer = SkipNextData( Buffer ); // advance past any soft seperators (tab or space)
 
-						if( CharacterTypeMap[ *Buffer ] == ECharType::EndOfElement )
+						if( IsCharTypeEqual(*Buffer, ECharType::EndOfElement) )
 						{
 							TCHAR Char = *Buffer++;
 							if( L'?' == Char )
@@ -495,7 +495,7 @@ protected:
 	
 	inline TCHAR* NextWhitespace( TCHAR* Buffer )
 	{
-		while( *Buffer && CharacterTypeMap[ *Buffer ] != ECharType::Whitespace )
+		while( *Buffer && !IsCharTypeEqual(*Buffer, ECharType::Whitespace))
 		{
 			Buffer++;
 		}
@@ -506,7 +506,7 @@ protected:
 
 	inline TCHAR* NextWhitespaceOrClose( TCHAR* Buffer, bool& IsClose )
 	{
-		while( *Buffer && CharacterTypeMap[ *Buffer ] != ECharType::Whitespace && *Buffer != L'>' )
+		while( *Buffer && !IsCharTypeEqual(*Buffer, ECharType::Whitespace) && *Buffer != L'>' )
 		{
 			Buffer++;
 		}
@@ -518,7 +518,7 @@ protected:
 
 	inline TCHAR* NextSeparator( TCHAR* Buffer )
 	{
-		while( *Buffer && CharacterTypeMap[ *Buffer ] != ECharType::Whitespace && *Buffer != L'=' )
+		while( *Buffer && !IsCharTypeEqual(*Buffer, ECharType::Whitespace) && *Buffer != L'=' )
 		{
 			Buffer++;
 		}
@@ -530,7 +530,7 @@ protected:
 	inline TCHAR* SkipNextData( TCHAR* Buffer )
 	{
 		// while we have data, and we encounter soft seperators or line feeds...
-		while( *Buffer && ( CharacterTypeMap[ *Buffer ] == ECharType::Whitespace || CharacterTypeMap[ *Buffer ] == ECharType::EndOfLine ) )
+		while( *Buffer && ( IsCharTypeEqual(*Buffer, ECharType::Whitespace) || IsCharTypeEqual(*Buffer, ECharType::EndOfLine) ) )
 		{
 			if( *Buffer == L'\n' )
 			{
@@ -594,6 +594,15 @@ protected:
 		EndOfLine,
 	};
 
+	inline bool IsCharTypeEqual( TCHAR Character, ECharType Type )
+	{
+		if (Character < 256)
+		{
+			return CharacterTypeMap[Character] == Type;
+		}
+		return Type == ECharType::Data;
+	}
+
 	/** Maps each ASCII character to the type of character we think it is */
 	ECharType CharacterTypeMap[ 256 ];
 
