remove unused dependency on CoreUObject module.

#rnx

diff --git a/Engine/Source/Runtime/XmlParser/XmlParser.Build.cs b/Engine/Source/Runtime/XmlParser/XmlParser.Build.cs
--- a/Engine/Source/Runtime/XmlParser/XmlParser.Build.cs
+++ b/Engine/Source/Runtime/XmlParser/XmlParser.Build.cs
@@ -10,7 +10,6 @@ public class XmlParser : ModuleRules
 			new string[] 
 			{ 
 				"Core",
-				"CoreUObject",
 			}
 			);
 	}
