Add a member function to XmlNode that allows to add Child Nodes with Attributes, add a member function to set attributes of a Node

diff --git a/Engine/Source/Runtime/XmlParser/Private/XmlNode.cpp b/Engine/Source/Runtime/XmlParser/Private/XmlNode.cpp
--- a/Engine/Source/Runtime/XmlParser/Private/XmlNode.cpp
+++ b/Engine/Source/Runtime/XmlParser/Private/XmlNode.cpp
@@ -76,12 +76,16 @@ const FString& FXmlNode::GetContent() const
 	return Content;
 }
 
-
 void FXmlNode::SetContent( const FString& InContent )
 {
 	Content = InContent;
 }
 
+void FXmlNode::SetAttributes(const TArray<FXmlAttribute>& InAttributes)
+{
+	Attributes = InAttributes;
+}
+
 FString FXmlNode::GetAttribute(const FString& InTag) const
 {
 	for(auto Iter(Attributes.CreateConstIterator()); Iter; Iter++)
@@ -94,11 +98,12 @@ FString FXmlNode::GetAttribute(const FString& InTag) const
 	return FString();
 }
 
-void FXmlNode::AppendChildNode(const FString& InTag, const FString& InContent)
+void FXmlNode::AppendChildNode(const FString& InTag, const FString& InContent, const TArray<FXmlAttribute>& InAttributes)
 {
 	auto NewNode = new FXmlNode;
 	NewNode->Tag = InTag;
 	NewNode->Content = InContent;
+	NewNode->Attributes = InAttributes;
 
 	auto NumChildren = Children.Num();
 	if (NumChildren != 0)
diff --git a/Engine/Source/Runtime/XmlParser/Public/XmlNode.h b/Engine/Source/Runtime/XmlParser/Public/XmlNode.h
--- a/Engine/Source/Runtime/XmlParser/Public/XmlNode.h
+++ b/Engine/Source/Runtime/XmlParser/Public/XmlNode.h
@@ -67,6 +67,8 @@ public:
 	const FString& GetContent() const;
 	/** Sets the new value of the node */
 	void SetContent(const FString& InContent);
+	/** Sets the attributes of the node */
+	void SetAttributes(const TArray<FXmlAttribute>& InAttributes);
 
 	/**
 	 * Gets all of the attributes in this node
@@ -80,8 +82,14 @@ public:
 
 	/** Gets an attribute that corresponds with the passed-in tag */
 	FString GetAttribute(const FString& InTag) const;
-	/** Adds a simple tag with content to the current node */
-	void AppendChildNode(const FString& InTag, const FString& InContent);
+
+	/**
+	 *  Adds a child node to this node  
+	 * @param	InTag				The tag of the child node
+	 * @param	InContent			(optional) The content of the child node
+	 * @param	InAttributes		(optional) An array of attributes of the child node
+	 */
+	void AppendChildNode(const FString& InTag, const FString& InContent = FString(), const TArray<FXmlAttribute>& InAttributes = TArray<FXmlAttribute>());
 private:
 
 	/** The list of children nodes */
