FXmlFile::LoadFile() and FXmlNode::Delete() no longer use unbounded recursion.

diff --git a/Engine/Source/Runtime/XmlParser/Private/XmlFile.cpp b/Engine/Source/Runtime/XmlParser/Private/XmlFile.cpp
--- a/Engine/Source/Runtime/XmlParser/Private/XmlFile.cpp
+++ b/Engine/Source/Runtime/XmlParser/Private/XmlFile.cpp
@@ -456,20 +456,7 @@ TArray<FString> FXmlFile::Tokenize(const TArray<FString>& Input)
 	return Tokens;
 }
 
-/** Checks if the passed string is an important tag operator */
-static bool IsTagOperator(FString ToCheck)
-{
-	if(ToCheck == TEXT("<") ||
-		ToCheck == TEXT(">") ||
-		ToCheck == TEXT("</") ||
-		ToCheck == TEXT("/>"))
-	{
-		return true;
-	}
-	return false;
-}
-
-void FXmlFile::AddAttribute(const FString& InToken, TArray<FXmlAttribute>& OutAttributes)
+static void AddAttribute(const FString& InToken, TArray<FXmlAttribute>& OutAttributes)
 {
 	int32 EqualsIdx;
 	if(InToken.FindChar(TEXT('='), EqualsIdx))
@@ -500,235 +487,244 @@ void FXmlFile::AddAttribute(const FString& InToken, TArray<FXmlAttribute>& OutAt
 	}
 }
 
-FXmlNode* FXmlFile::CreateNodeRecursive(const TArray<FString>& Tokens, int32 StartIndex, int32* OutNextIndex)
+enum class ETryParseTagResult
 {
-	// Algorithm (Draft):
-	//  - First found token should always be a '<'
-	//  - Extract tag
-	//  - Check next token
-	//    - If '<', recursively deepen
-	//    - If string token, parse until '<' found to recursively deepen
-	//  - Continue parsing until </tag> for self is found
-	//  - Return own constructed node (and index of next starting point
+	Error,
+	None,
+	ParsedStartTag,         // <key>
+	ParsedEmptyElementTag,  // <key />
+	ParsedEndTag,           // </key>
+	NotATag                 // content - iterator is not advanced past this
+};
+
+static ETryParseTagResult TryParseTag(FString& OutTag, TArray<FXmlAttribute>& Attributes, const FString*& InOutToken, const FString* LastToken)
+{
+	const FString* Token = InOutToken;
+
+	if (Token == LastToken)
+	{
+		return ETryParseTagResult::None;
+	}
 
-	const int32 RecursiveStartIndex = StartIndex;
-	int32 SavedIndex = StartIndex;
+	if (*Token == TEXT(">") || *Token == TEXT("/>"))
+	{
+		// Error: Found an unexpected end-of-tag (ex: >list>)
+		return ETryParseTagResult::Error;
+	}
 
-	// Get the tag & any attributes
-	FString Tag;
-	TArray<FXmlAttribute> Attributes;
-	bool bInTag = false;
-	for(int32 i = StartIndex; i < Tokens.Num(); ++i)
+	// Found an end tag
+	if (*Token == TEXT("</"))
 	{
-		if(bCreationFailed)
-		{
-			break;
-		}
+		++Token;
 
-		if(Tokens[i].Len() == 0 || Tokens[i] == TEXT("\n") || Tokens[i] == TEXT("\r") || Tokens[i] == TEXT("\n\r"))
+		if (Token == LastToken || *Token == TEXT("<") || *Token == TEXT("</") || *Token == TEXT(">") || *Token == TEXT("/>"))
 		{
-			continue;
+			// Error: Didn't find a tag name before the end of the tokens (ex: </>)
+			return ETryParseTagResult::Error;
 		}
 
-		// Looking for tag start
-		if(!bInTag)
+		OutTag = *Token++;
+
+		for (;;)
 		{
-			if(Tokens[i] == TEXT("<"))
+			if (Token == LastToken || *Token == TEXT("<") || *Token == TEXT("</") || *Token == TEXT("/>"))
 			{
-				bInTag = true;
+				// Error: Didn't find a closing bracket (ex: </EndToken<)
+				return ETryParseTagResult::Error;
 			}
 
-			else
+			if (*Token++ == TEXT(">"))
 			{
-				// Error: Found text before any operators (ex: plist>)
-				bCreationFailed = true;
-				ErrorMessage = NSLOCTEXT("XmlParser", "MalformedXMLFile", "Malformed Xml File").ToString();
+				InOutToken = Token;
+				return ETryParseTagResult::ParsedEndTag;
 			}
+
+			// Don't error if we find text after the final tag - they may be like options or modifiers or some such
 		}
+	}
+
+	if (*Token != TEXT("<"))
+	{
+		// Found a non-tag
+		// Don't advance the Token iterator because we want the caller to be able to examine it
+		InOutToken = Token;
+		return ETryParseTagResult::NotATag;
+	}
+
+	++Token;
 
-		// Already found tag start, reading inside of it now
-		else if(bInTag)
+	if (Token == LastToken || *Token == TEXT("<") || *Token == TEXT("</") || *Token == TEXT(">") || *Token == TEXT("/>"))
+	{
+		// Error: Didn't find a tag name before the end of the tokens (ex: <)
+		return ETryParseTagResult::Error;
+	}
+
+	// Copy tag to output
+	OutTag = *Token++;
+
+	for (;;)
+	{
+		if (Token == LastToken)
 		{
-			// Text for the tag
-			if(!IsTagOperator(Tokens[i]))
-			{
-				if(Tag.Len())
-				{
-					// this could be an attribute
-					AddAttribute(Tokens[i], Attributes);
-				}
-				else
-				{
-					Tag = Tokens[i];
-				}
-			}
+			// Error: Didn't find the closing tag (ex: <key)
+			return ETryParseTagResult::Error;
+		}
 
-			// Found another tag operator
-			else
-			{
-				// Closing the tag starting
-				if(Tokens[i] == TEXT(">"))
-				{
-					SavedIndex = i + 1; // We want to find the content starting here possibly
-					break;
-				}
+		const FString& TokenStr = *Token++;
 
-				// Closing the tag and also finalizing it
-				else if(Tokens[i] == TEXT("/>"))
-				{
-					FXmlNode* NewNode = new FXmlNode();
-					NewNode->Tag = Tag;
-					NewNode->Attributes = Attributes;
-					if (OutNextIndex != nullptr)
-					{
-						*OutNextIndex = i + 1;
-					}
-					return NewNode;
-				}
+		// Closing the tag and also finalizing it
+		if (TokenStr == TEXT("/>"))
+		{
+			InOutToken = Token;
+			return ETryParseTagResult::ParsedEmptyElementTag;
+		}
 
-				// Error
-				else
-				{
-					// Error: malformed file (ex: <key<)
-					bCreationFailed = true;
-					ErrorMessage = NSLOCTEXT("XmlParser", "MalformedXMLFile", "Malformed Xml File").ToString();
-				}
-			}
+		// Closing the tag starting
+		if (TokenStr == TEXT(">"))
+		{
+			InOutToken = Token;
+			return ETryParseTagResult::ParsedStartTag;
 		}
-	}
-	
-	// Create a new node for the current tag
-	FXmlNode* NewNode = new FXmlNode();
-	NewNode->Tag = Tag;
-	NewNode->Attributes = Attributes;
 
-	// Got the tag. Continue and try to get the content
-	FString Content;
-	FString FinalTag;
-	bInTag = false;
-	for(int32 i = SavedIndex; i < Tokens.Num(); ++i)
-	{
-		if(bCreationFailed)
+		// Error
+		if (TokenStr == TEXT("<") || TokenStr == TEXT("</"))
 		{
-			break;
+			// Error: malformed file (ex: <key<)
+			return ETryParseTagResult::Error;
 		}
 
-		if(Tokens[i].Len() == 0 || Tokens[i] == TEXT("\n") || Tokens[i] == TEXT("\r") || Tokens[i] == TEXT("\n\r"))
+		// Assume this is an attribute
+		AddAttribute(TokenStr, Attributes);
+	}
+}
+
+FXmlNode* FXmlFile::CreateRootNode(TArrayView<const FString> Tokens)
+{
+	// Algorithm (Draft):
+	//  - First found token should always be a '<'
+	//  - Extract tag
+	//  - Check next token
+	//    - If '<', recursively deepen
+	//    - If string token, parse until '<' found to recursively deepen
+	//  - Continue parsing until </tag> for self is found
+	//  - Return own constructed node (and index of next starting point
+
+	// Stack of nodes currently being parsed - Last() is the current node being parsed
+	TArray<FXmlNode*> ParsingNodeStack;
+	ON_SCOPE_EXIT
+	{
+		// Clean up partially-parsed nodes - only the root node needs deletion because nested nodes
+		// will have been added to the Children list of that node, and Delete already cleans those up.
+		if (ParsingNodeStack.Num() != 0)
 		{
-			continue;
+			delete ParsingNodeStack[0];
 		}
+	};
+
+	const FString* Token     = Tokens.GetData();
+	const FString* LastToken = Token + Tokens.Num();
+	for (;;)
+	{
+		FString               Tag;
+		TArray<FXmlAttribute> Attributes;
 
-		if(!bInTag)
+		ETryParseTagResult ParseResult = TryParseTag(Tag, Attributes, Token, LastToken);
+		if (ParseResult == ETryParseTagResult::None)
 		{
-			// Found the start of another tag
-			if(Tokens[i] == TEXT("<"))
+			// Run out of tokens mid-parse
+
+			if (!ParsingNodeStack.IsEmpty())
 			{
-				// Recursively enter function creating a child at the new tag
-				FXmlNode* Child = nullptr;
-				if (i > RecursiveStartIndex)
-				{
-					Child = CreateNodeRecursive(Tokens, i, &SavedIndex);
-				}
+				bCreationFailed = true;
+				ErrorMessage = NSLOCTEXT("XmlParser", "MalformedXMLFile", "Malformed Xml File").ToString();
+				return nullptr;
+			}
 
-				// Save child to parent
-				if(Child != nullptr)
-				{
-					NewNode->Children.Add(Child);
-				}
+			// Should only be possible if there was no content - original parser
+			// would return nullptr without setting bCreationFailed and ErrorMessage
+			return nullptr;
+		}
 
-				// Error creating a child
-				else
-				{
-					// Some error occurred in creating a child. The algorithm should abort
-					bCreationFailed = true;
-					ErrorMessage = NSLOCTEXT("XmlParser", "MalformedXMLFile", "Malformed Xml File").ToString();
-					break;
-				}
+		if (ParseResult == ETryParseTagResult::Error)
+		{
+			// Malformed tag
 
-				// Continue algorithm at the returned index
-				i = SavedIndex - 1; // (-1) because continuing will increment i;
-				continue;
-			}
+			bCreationFailed = true;
+			ErrorMessage = NSLOCTEXT("XmlParser", "MalformedXMLFile", "Malformed Xml File").ToString();
+			return nullptr;
+		}
+
+		if (ParseResult == ETryParseTagResult::ParsedEndTag)
+		{
+			// Found an end tag - pop the current node and return to parsing the parent
 
-			// Found what should be the end of the current tag
-			else if(Tokens[i] == TEXT("</"))
+			if (ParsingNodeStack.IsEmpty())
 			{
-				// Continue, getting the tag and checking to make sure it matched
-				bInTag = true;
+				// Error: encountered an end tag before we found any opening tag
+				bCreationFailed = true;
+				ErrorMessage = NSLOCTEXT("XmlParser", "MalformedXMLFile", "Malformed Xml File").ToString();
 			}
 
-			// Error
-			else if(IsTagOperator(Tokens[i]))
+			if (ParsingNodeStack.Last()->Tag != Tag)
 			{
-				// Error: Invalid token such as '<key>>'
+				// Error: Tag and ending tag don't match
 				bCreationFailed = true;
 				ErrorMessage = NSLOCTEXT("XmlParser", "MalformedXMLFile", "Malformed Xml File").ToString();
+				return nullptr;
 			}
 
-			// Not an operator, save the text 
-			else
+			FXmlNode* FullyParsedNode = ParsingNodeStack.Pop();
+
+			if (ParsingNodeStack.IsEmpty())
 			{
-				if(Content.Len() > 0)
-				{
-					Content += TEXT(" ");
-				}
-				Content += Tokens[i];
+				// Found the end tag of the root node - stop
+				return FullyParsedNode;
 			}
+
+			continue;
 		}
 
-		// In the finalizing tag
-		else
+		if (ParseResult == ETryParseTagResult::NotATag)
 		{
-			// Text for the tag
-			if(!IsTagOperator(Tokens[i]))
+			// Found some content - add it to the content of the current node
+
+			if (ParsingNodeStack.IsEmpty())
 			{
-				if(FinalTag.Len())
-				{
-					// Ignore this bit of text then (they may be like options or modifiers or some such)
-				}
-				else
-				{
-					FinalTag = Tokens[i];
-				}
+				// Error: encountered content tokens outside of a nested tag
+				bCreationFailed = true;
+				ErrorMessage = NSLOCTEXT("XmlParser", "MalformedXMLFile", "Malformed Xml File").ToString();
 			}
 
-			else
-			{
-				// Found the end of the final tag
-				if(Tokens[i] == TEXT(">"))
-				{
-					// Error: Tag and ending tag don't match
-					if(Tag != FinalTag)
-					{
-						// Error: Tag and ending tag don't match
-						bCreationFailed = true;
-						ErrorMessage = NSLOCTEXT("XmlParser", "MalformedXMLFile", "Malformed Xml File").ToString();
-					}
-					
-					// Check to make sure tags match and create the node and return it
-					SavedIndex = i + 1;
-					if(OutNextIndex != nullptr) // Might be null because of the top-level node not sending it in
-					{
-						*OutNextIndex = SavedIndex;
-					}
-					NewNode->Content = Content;
-					return NewNode;
-				}
+			FString& Content = ParsingNodeStack.Last()->Content;
 
-				// Error
-				else
-				{
-					// Error: malformed file (ex: <key>stuff</key/>)
-					bCreationFailed = true;
-					ErrorMessage = NSLOCTEXT("XmlParser", "MalformedXMLFile", "Malformed Xml File").ToString();
-				}
+			if (Content.Len() > 0)
+			{
+				Content += TEXT(" ");
 			}
+			Content += *Token++;
+
+			continue;
 		}
-	}
 
-	// Something bad happened, return NULL!
-	delete NewNode;
-	return nullptr;
+		// Found a nested node
+
+		check(ParseResult == ETryParseTagResult::ParsedStartTag || ParseResult == ETryParseTagResult::ParsedEmptyElementTag);
+
+		FXmlNode* NewNode = new FXmlNode();
+		NewNode->Tag = MoveTemp(Tag);
+		NewNode->Attributes = MoveTemp(Attributes);
+
+		// If we are already parsing nodes, this is a child node, so add it to the parent
+		if (!ParsingNodeStack.IsEmpty())
+		{
+			ParsingNodeStack.Last()->Children.Add(NewNode);
+		}
+
+		if (ParseResult == ETryParseTagResult::ParsedStartTag)
+		{
+			ParsingNodeStack.Add(NewNode);
+		}
+	}
 }
 
 void FXmlFile::HookUpNextPtrs(FXmlNode* Node)
@@ -753,7 +749,7 @@ void FXmlFile::CreateNodes(const TArray<FString>& Tokens)
 {
 	// Assumption..There is only 1 top-level node which contains everything inside of it
 	bCreationFailed = false;
-	FXmlNode* Root = CreateNodeRecursive(Tokens, 0);
+	FXmlNode* Root = CreateRootNode(Tokens);
 
 	if(Root)
 	{
diff --git a/Engine/Source/Runtime/XmlParser/Private/XmlNode.cpp b/Engine/Source/Runtime/XmlParser/Private/XmlNode.cpp
--- a/Engine/Source/Runtime/XmlParser/Private/XmlNode.cpp
+++ b/Engine/Source/Runtime/XmlParser/Private/XmlNode.cpp
@@ -15,14 +15,20 @@ const FString& FXmlAttribute::GetValue() const
 
 void FXmlNode::Delete()
 {
-	const int32 ChildCount = Children.Num();
-	for(int32 ChildIndex = 0; ChildIndex < ChildCount; ++ChildIndex)
+	TArray<FXmlNode*> ToDelete = MoveTemp(Children);
+	check(Children.IsEmpty());
+
+	for (int32 Index = 0; Index != ToDelete.Num(); ++Index)
+	{
+		FXmlNode* NodeToDelete = ToDelete[Index];
+		ToDelete.Append(MoveTemp(NodeToDelete->Children));
+		check(NodeToDelete->Children.IsEmpty());
+	}
+
+	for (FXmlNode* Node : ToDelete)
 	{
-		check(Children[ChildIndex] != nullptr);
-		Children[ChildIndex]->Delete();
-		delete Children[ChildIndex];
+		delete Node;
 	}
-	Children.Empty();
 }
 
 const FXmlNode* FXmlNode::GetNextNode() const
diff --git a/Engine/Source/Runtime/XmlParser/Public/XmlFile.h b/Engine/Source/Runtime/XmlParser/Public/XmlFile.h
--- a/Engine/Source/Runtime/XmlParser/Public/XmlFile.h
+++ b/Engine/Source/Runtime/XmlParser/Public/XmlFile.h
@@ -89,9 +89,9 @@ private:
 	 * Add an attribute to the passed-in array.
 	 * This makes the assumption that an attribute comes in as one 'token' (i.e. no spaces between tag="value").
 	 */
-	void AddAttribute(const FString& InToken, TArray<FXmlAttribute>& OutAttributes);
+	static void AddAttribute(const FString& InToken, TArray<FXmlAttribute>& OutAttributes);
 	/** Recursively creates nodes from the list of tokens */
-	FXmlNode* CreateNodeRecursive(const TArray<FString>& Tokens, int32 StartIndex, int32* OutNextIndex = nullptr);
+	FXmlNode* CreateRootNode(TArrayView<const FString> Tokens);
 	/** Hook next ptrs up recursively */
 	void HookUpNextPtrs(FXmlNode* Node);
 	/** Creates the internal file representation as a bunch of FXmlNode's */
