Fixed FXmlFile not early-returning when parsing some invalid XML.
Fixed FXmlFile not handing empty element root nodes (<Root />).

diff --git a/Engine/Source/Runtime/XmlParser/Private/XmlFile.cpp b/Engine/Source/Runtime/XmlParser/Private/XmlFile.cpp
--- a/Engine/Source/Runtime/XmlParser/Private/XmlFile.cpp
+++ b/Engine/Source/Runtime/XmlParser/Private/XmlFile.cpp
@@ -663,6 +663,7 @@ FXmlNode* FXmlFile::CreateRootNode(TArrayView<const FString> Tokens)
 				// Error: encountered an end tag before we found any opening tag
 				bCreationFailed = true;
 				ErrorMessage = NSLOCTEXT("XmlParser", "MalformedXMLFile", "Malformed Xml File").ToString();
+				return nullptr;
 			}
 
 			if (ParsingNodeStack.Last()->Tag != Tag)
@@ -693,6 +694,7 @@ FXmlNode* FXmlFile::CreateRootNode(TArrayView<const FString> Tokens)
 				// Error: encountered content tokens outside of a nested tag
 				bCreationFailed = true;
 				ErrorMessage = NSLOCTEXT("XmlParser", "MalformedXMLFile", "Malformed Xml File").ToString();
+				return nullptr;
 			}
 
 			FString& Content = ParsingNodeStack.Last()->Content;
@@ -719,6 +721,11 @@ FXmlNode* FXmlFile::CreateRootNode(TArrayView<const FString> Tokens)
 		{
 			ParsingNodeStack.Last()->Children.Add(NewNode);
 		}
+		else if (ParseResult == ETryParseTagResult::ParsedEmptyElementTag)
+		{
+			// If we just parsed an empty element root node, stop
+			return NewNode;
+		}
 
 		if (ParseResult == ETryParseTagResult::ParsedStartTag)
 		{
