Fix XML header search

#rnx

diff --git a/Engine/Source/Runtime/XmlParser/Private/XmlFile.cpp b/Engine/Source/Runtime/XmlParser/Private/XmlFile.cpp
--- a/Engine/Source/Runtime/XmlParser/Private/XmlFile.cpp
+++ b/Engine/Source/Runtime/XmlParser/Private/XmlFile.cpp
@@ -170,7 +170,7 @@ void FXmlFile::PreProcessInput(TArray<FString>& Input)
 	for(int32 i = 0; i < Input.Num(); ++i)
 	{
 		// Find <!DOCTYPE or <?xml and remove those lines
-		if(Input[i].StartsWith(TEXT("<!DOCTYPE")) || Input[i].StartsWith(TEXT("<?xml")))
+		if(Input[i].StartsWith(TEXT("<!DOCTYPE"), ESearchCase::IgnoreCase) || Input[i].StartsWith(TEXT("<?xml"), ESearchCase::IgnoreCase))
 		{
 			Input[i] = TEXT("");
 		}
