XmlParser | Make the XML parser preserve newlines

diff --git a/Engine/Source/Runtime/XmlParser/Private/XmlFile.cpp b/Engine/Source/Runtime/XmlParser/Private/XmlFile.cpp
--- a/Engine/Source/Runtime/XmlParser/Private/XmlFile.cpp
+++ b/Engine/Source/Runtime/XmlParser/Private/XmlFile.cpp
@@ -1,6 +1,8 @@
 // Copyright Epic Games, Inc. All Rights Reserved.
 
 #include "XmlFile.h"
+
+#include "Algo/AllOf.h"
 #include "Misc/FileHelper.h"
 #include "Misc/ScopeExit.h"
 FXmlFile::FXmlFile(const FString& InFile, EConstructMethod::Type ConstructMethod)
@@ -39,6 +41,12 @@ bool FXmlFile::LoadFile(const FString& InFile, EConstructMethod::Type ConstructM
 	// Pre-process the input
 	PreProcessInput(Input);
 
+	// Reintroduce line endings which are stripped by ParseIntoArrayLines / LoadFileToStringArray.
+	for (FString& Line : Input)
+	{
+		Line.AppendChar('\n');
+	}
+
 	// Tokenize the input
 	TArray<FString> Tokens = Tokenize(Input);
 
@@ -347,7 +355,7 @@ void FXmlFile::Tokenize(FStringView Input, TArray<FString>& Tokens)
 	{
 		TCHAR Ch = *Ptr;
 
-		if(IsWhiteSpace(Ch) && !bInQuote)
+		if (IsWhiteSpace(Ch) && Ch != TEXT('\n') && !bInQuote)
 		{
 			// End the current token 
 			if(WorkingToken.Len())
@@ -691,6 +699,13 @@ FXmlNode* FXmlFile::CreateRootNode(TArrayView<const FString> Tokens)
 
 			if (ParsingNodeStack.IsEmpty())
 			{
+				if (Algo::AllOf(*Token, IsWhiteSpace))
+				{
+					// Whitespace is allowed outside the root node.
+					++Token;
+					continue;
+				}
+
 				// Error: encountered content tokens outside of a nested tag
 				bCreationFailed = true;
 				ErrorMessage = NSLOCTEXT("XmlParser", "MalformedXMLFile", "Malformed Xml File").ToString();
@@ -699,7 +714,7 @@ FXmlNode* FXmlFile::CreateRootNode(TArrayView<const FString> Tokens)
 
 			FString& Content = ParsingNodeStack.Last()->Content;
 
-			if (Content.Len() > 0)
+			if (Content.Len() > 0 && !Content.EndsWith(TEXT("\n")))
 			{
 				Content += TEXT(" ");
 			}
