Reverted accidental submit of XmlFile.cpp.

diff --git a/Engine/Source/Runtime/XmlParser/Private/XmlFile.cpp b/Engine/Source/Runtime/XmlParser/Private/XmlFile.cpp
--- a/Engine/Source/Runtime/XmlParser/Private/XmlFile.cpp
+++ b/Engine/Source/Runtime/XmlParser/Private/XmlFile.cpp
@@ -346,7 +346,7 @@ void FXmlFile::Tokenize(FStringView Input, TArray<FString>& Tokens)
 		}
 	};
 
-	enum TOKENTYPE { UNKNOWN, OPERATOR, STRING } Type = UNKNOWN;
+	enum TOKENTYPE { OPERATOR, STRING } Type = STRING;
 	bool bInQuote = false;
 
 	const TCHAR* PtrStart = GetData(Input);
@@ -367,99 +367,88 @@ void FXmlFile::Tokenize(FStringView Input, TArray<FString>& Tokens)
 			continue;
 		}
 
-		switch (Type)
+		// Mark the start of a token
+		if(!WorkingToken.Len())
 		{
-			case UNKNOWN:
+			WorkingToken += Ch;
+			if(CheckTagOperator(PtrStart, Ptr, PtrEnd))
 			{
-				WorkingToken += Ch;
-				if(CheckTagOperator(PtrStart, Ptr, PtrEnd))
+				// Add the working token if it's final (ie: ends with '>')
+				if (Ch == TCHAR('>'))
 				{
-					// Add the working token if it's final (ie: ends with '>')
-					if (Ch == TCHAR('>'))
-					{
-						Tokens.Add(MoveTemp(WorkingToken));
-						checkSlow(WorkingToken.Len() == 0);
-					}
-					else
-					{
-						Type = OPERATOR;
-					}
+					Tokens.Add(MoveTemp(WorkingToken));
+					checkSlow(WorkingToken.Len() == 0);
 				}
 				else
 				{
-					Type = STRING;
+					Type = OPERATOR;
 				}
 			}
-			break;
+			else
+			{
+				Type = STRING;
+			}
 
-			case OPERATOR:
+			continue;
+		}
+
+		// Already in a token, so continue parsing
+		if(Type == OPERATOR)
+		{
+			// Still the tag, so add it to the working token
+			if(CheckTagOperator(PtrStart, Ptr, PtrEnd))
 			{
-				// Still the tag, so add it to the working token
-				if(CheckTagOperator(PtrStart, Ptr, PtrEnd))
-				{
-					WorkingToken += Ch;
-
-					// Add the working token if it's final (ie: ends with '>')
-					if (Ch == TCHAR('>'))
-					{
-						Tokens.Add(MoveTemp(WorkingToken));
-						checkSlow(WorkingToken.Len() == 0);
-						Type = UNKNOWN;
-					}
-				}
+				WorkingToken += Ch;
 
-				// Not a tag operator anymore, so add the old token and start a new one
-				else
+				// Add the working token if it's final (ie: ends with '>')
+				if (Ch == TCHAR('>'))
 				{
 					Tokens.Add(MoveTemp(WorkingToken));
 					checkSlow(WorkingToken.Len() == 0);
-					WorkingToken += Ch;
-					Type = STRING;
 				}
 			}
-			break;
 
-			case STRING:
+			// Not a tag operator anymore, so add the old token and start a new one
+			else
 			{
-				if (IsQuote(Ch))
-				{
-					bInQuote = !bInQuote;
-				}
+				Tokens.Add(MoveTemp(WorkingToken));
+				checkSlow(WorkingToken.Len() == 0);
+				WorkingToken += Ch;
+				Type = STRING;
+			}
+		}
+		else // STRING
+		{
+			if (IsQuote(Ch))
+			{
+				bInQuote = !bInQuote;
+			}
 
-				// Still a string. Allow '>' within a string
-				if(!CheckTagOperator(PtrStart, Ptr, PtrEnd) || (bInQuote && Ch == TCHAR('>')))
-				{
-					WorkingToken += Ch;
-				}
+			// Still a string. Allow '>' within a string
+			if(!CheckTagOperator(PtrStart, Ptr, PtrEnd) || (bInQuote && Ch == TCHAR('>')))
+			{
+				WorkingToken += Ch;
+			}
 
-				// Moving back to operator
-				else
+			// Moving back to operator
+			else
+			{
+				Tokens.Add(MoveTemp(WorkingToken));
+				checkSlow(WorkingToken.Len() == 0);
+				WorkingToken += Ch;
+				bInQuote = false;
+
+				// Add the working token if it's final (ie: ends with '>')
+				if (Ch == TCHAR('>'))
 				{
 					Tokens.Add(MoveTemp(WorkingToken));
 					checkSlow(WorkingToken.Len() == 0);
-					WorkingToken += Ch;
-					bInQuote = false;
-
-					// Add the working token if it's final (ie: ends with '>')
-					if (Ch == TCHAR('>'))
-					{
-						Tokens.Add(MoveTemp(WorkingToken));
-						checkSlow(WorkingToken.Len() == 0);
-						Type = UNKNOWN;
-					}
-					else
-					{
-						Type = OPERATOR;
-					}
+				}
+				else
+				{
+					Type = OPERATOR;
 				}
 			}
-			break;
-
-			default:
-			{
-				checkSlow("Unhandled state");
-			}
-			break;
 		}
 	}
 }
