Back out changelist 29816891 to fix parsing of SVN server responses.

In FXmlFile::CreateRootNode the check `ParsingNodeStack.Last()->Tag != Tag` was failing in the case of e.g. a start tag "entry\n" mismatched with an end tag of "entry".  Attribute value tokens also had trailing newline chars injected that were causing further issues with SVN commands.  An accumulation of newline chars was also present in the node Content.

diff --git a/Engine/Source/Runtime/XmlParser/Private/XmlFile.cpp b/Engine/Source/Runtime/XmlParser/Private/XmlFile.cpp
--- a/Engine/Source/Runtime/XmlParser/Private/XmlFile.cpp
+++ b/Engine/Source/Runtime/XmlParser/Private/XmlFile.cpp
@@ -1,8 +1,6 @@
 // Copyright Epic Games, Inc. All Rights Reserved.
 
 #include "XmlFile.h"
-
-#include "Algo/AllOf.h"
 #include "Misc/FileHelper.h"
 #include "Misc/ScopeExit.h"
 FXmlFile::FXmlFile(const FString& InFile, EConstructMethod::Type ConstructMethod)
@@ -41,12 +39,6 @@ bool FXmlFile::LoadFile(const FString& InFile, EConstructMethod::Type ConstructM
 	// Pre-process the input
 	PreProcessInput(Input);
 
-	// Reintroduce line endings which are stripped by ParseIntoArrayLines / LoadFileToStringArray.
-	for (FString& Line : Input)
-	{
-		Line.AppendChar('\n');
-	}
-
 	// Tokenize the input
 	TArray<FString> Tokens = Tokenize(Input);
 
@@ -355,7 +347,7 @@ void FXmlFile::Tokenize(FStringView Input, TArray<FString>& Tokens)
 	{
 		TCHAR Ch = *Ptr;
 
-		if (IsWhiteSpace(Ch) && Ch != TEXT('\n') && !bInQuote)
+		if(IsWhiteSpace(Ch) && !bInQuote)
 		{
 			// End the current token 
 			if(WorkingToken.Len())
@@ -699,13 +691,6 @@ FXmlNode* FXmlFile::CreateRootNode(TArrayView<const FString> Tokens)
 
 			if (ParsingNodeStack.IsEmpty())
 			{
-				if (Algo::AllOf(*Token, IsWhiteSpace))
-				{
-					// Whitespace is allowed outside the root node.
-					++Token;
-					continue;
-				}
-
 				// Error: encountered content tokens outside of a nested tag
 				bCreationFailed = true;
 				ErrorMessage = NSLOCTEXT("XmlParser", "MalformedXMLFile", "Malformed Xml File").ToString();
@@ -714,7 +699,7 @@ FXmlNode* FXmlFile::CreateRootNode(TArrayView<const FString> Tokens)
 
 			FString& Content = ParsingNodeStack.Last()->Content;
 
-			if (Content.Len() > 0 && !Content.EndsWith(TEXT("\n")))
+			if (Content.Len() > 0)
 			{
 				Content += TEXT(" ");
 			}
