Set UE_DEFINE_LEGACY_MATH_CONSTANT_MACRO_NAMES=0 for XmlParser.
Disabled PCHs because it will use Core's shared PCH and will clash with the macro redefinition.

diff --git a/Engine/Source/Runtime/XmlParser/XmlParser.Build.cs b/Engine/Source/Runtime/XmlParser/XmlParser.Build.cs
--- a/Engine/Source/Runtime/XmlParser/XmlParser.Build.cs
+++ b/Engine/Source/Runtime/XmlParser/XmlParser.Build.cs
@@ -12,5 +12,10 @@ public class XmlParser : ModuleRules
 				"Core",
 			});
 		UnsafeTypeCastWarningLevel = WarningLevel.Error;
+
+		// PCH is disabled here otherwise the module will use Core's shared PCH and fail to compile due
+		// to a clash with UE_DEFINE_LEGACY_MATH_CONSTANT_MACRO_NAMES being defined in UnrealMathUtility.h.
+		PCHUsage = PCHUsageMode.NoPCHs;
+		PrivateDefinitions.Add("UE_DEFINE_LEGACY_MATH_CONSTANT_MACRO_NAMES=0");
 	}
 }
